- Remove UE4 from LiveLink Provider

diff --git a/Engine/Source/Runtime/LiveLinkMessageBusFramework/Public/LiveLinkProvider.h b/Engine/Source/Runtime/LiveLinkMessageBusFramework/Public/LiveLinkProvider.h
--- a/Engine/Source/Runtime/LiveLinkMessageBusFramework/Public/LiveLinkProvider.h
+++ b/Engine/Source/Runtime/LiveLinkMessageBusFramework/Public/LiveLinkProvider.h
@@ -43,7 +43,7 @@ public:
 	virtual ~ILiveLinkProvider() {}
 
 	/**
-	 * Send, to UE4, the static data of a subject.
+	 * Send, to UE, the static data of a subject.
 	 * @param SubjectName	The name of the subject
 	 * @param Role			The Live Link role of the subject. The StaticData type should match the role's data.
 	 * @param StaticData	The static data of the subject.
