- Bumping livelink version for ue5 and fixing advertised version in LL message bus

diff --git a/Engine/Plugins/Animation/LiveLink/Source/LiveLink/Private/LiveLinkMessageBusSource.cpp b/Engine/Plugins/Animation/LiveLink/Source/LiveLink/Private/LiveLinkMessageBusSource.cpp
--- a/Engine/Plugins/Animation/LiveLink/Source/LiveLink/Private/LiveLinkMessageBusSource.cpp
+++ b/Engine/Plugins/Animation/LiveLink/Source/LiveLink/Private/LiveLinkMessageBusSource.cpp
@@ -253,7 +253,7 @@ FORCEINLINE void FLiveLinkMessageBusSource::UpdateConnectionLastActive()
 void FLiveLinkMessageBusSource::SendConnectMessage()
 {
 	FLiveLinkConnectMessage* ConnectMessage = new FLiveLinkConnectMessage();
-	ConnectMessage->LiveLinkVersion = 2;
+	ConnectMessage->LiveLinkVersion = ILiveLinkClient::LIVELINK_VERSION;
 	MessageEndpoint->Send(ConnectMessage, ConnectionAddress);
 	FLiveLinkHeartbeatEmitter& HeartbeatEmitter = ILiveLinkModule::Get().GetHeartbeatEmitter();
 	HeartbeatEmitter.StartHeartbeat(ConnectionAddress, MessageEndpoint);
diff --git a/Engine/Source/Runtime/LiveLinkInterface/Public/ILiveLinkClient.h b/Engine/Source/Runtime/LiveLinkInterface/Public/ILiveLinkClient.h
--- a/Engine/Source/Runtime/LiveLinkInterface/Public/ILiveLinkClient.h
+++ b/Engine/Source/Runtime/LiveLinkInterface/Public/ILiveLinkClient.h
@@ -111,7 +111,7 @@ PRAGMA_ENABLE_DEPRECATION_WARNINGS
 class LIVELINKINTERFACE_API ILiveLinkClient : public ILiveLinkClient_Base_DEPRECATED
 {
 public:
-	static const int32 LIVELINK_VERSION = 2;
+	static const int32 LIVELINK_VERSION = 3;
 
 	virtual ~ILiveLinkClient() {}
 
diff --git a/Engine/Source/Runtime/LiveLinkMessageBusFramework/Private/LiveLinkProvider.cpp b/Engine/Source/Runtime/LiveLinkMessageBusFramework/Private/LiveLinkProvider.cpp
--- a/Engine/Source/Runtime/LiveLinkMessageBusFramework/Private/LiveLinkProvider.cpp
+++ b/Engine/Source/Runtime/LiveLinkMessageBusFramework/Private/LiveLinkProvider.cpp
@@ -14,7 +14,7 @@
 
 DEFINE_LOG_CATEGORY_STATIC(LogLiveLinkMessageBus, Warning, All);
 
-static const int32 LIVELINK_SupportedVersion = 2;
+static const int32 LIVELINK_SupportedVersion = 3;
 
 
 FName FLiveLinkMessageAnnotation::SubjectAnnotation = TEXT("SubjectName");
@@ -464,7 +464,7 @@ PRAGMA_ENABLE_DEPRECATION_WARNINGS
 
 void FLiveLinkProvider::HandlePingMessage(const FLiveLinkPingMessage& Message, const TSharedRef<IMessageContext, ESPMode::ThreadSafe>& Context)
 {
-	if (Message.LiveLinkVersion < LIVELINK_SupportedVersion)
+	if (Message.LiveLinkVersion != LIVELINK_SupportedVersion)
 	{
 		UE_LOG(LogLiveLinkMessageBus, Warning, TEXT("A unsupported version of LiveLink is trying to communicate. Requested version: '%d'. Supported version: '%d'."), Message.LiveLinkVersion, LIVELINK_SupportedVersion)
 		return;
@@ -477,7 +477,7 @@ void FLiveLinkProvider::HandleConnectMessage(const FLiveLinkConnectMessage& Mess
 {
 	FScopeLock Lock(&CriticalSection);
 
-	if (Message.LiveLinkVersion < LIVELINK_SupportedVersion)
+	if (Message.LiveLinkVersion != LIVELINK_SupportedVersion)
 	{
 		UE_LOG(LogLiveLinkMessageBus, Error, TEXT("A unsupported version of LiveLink is trying to connect. Requested version: '%d'. Supported version: '%d'."), Message.LiveLinkVersion, LIVELINK_SupportedVersion)
 		return;
