- Rollback LiveLink MessageBus version to support new backward compabitility with older UE

diff --git a/Engine/Plugins/Animation/LiveLink/Source/LiveLink/Private/LiveLinkMessageBusDiscoveryManager.cpp b/Engine/Plugins/Animation/LiveLink/Source/LiveLink/Private/LiveLinkMessageBusDiscoveryManager.cpp
--- a/Engine/Plugins/Animation/LiveLink/Source/LiveLink/Private/LiveLinkMessageBusDiscoveryManager.cpp
+++ b/Engine/Plugins/Animation/LiveLink/Source/LiveLink/Private/LiveLinkMessageBusDiscoveryManager.cpp
@@ -110,13 +110,8 @@ void FLiveLinkMessageBusDiscoveryManager::HandlePongMessage(const FLiveLinkPongM
 
 	if (Message.PollRequest == LastPingRequest)
 	{
-		bool bIsValidProvider = true;
-
-		if (Message.LiveLinkVersion < ILiveLinkClient::LIVELINK_VERSION) // UE4 Remote Clients always send 1 as the LiveLink version
-		{
-			bIsValidProvider = false;
-		}
-
+		// Verify Message.LiveLinkVersion to consider validity of discovered provider. Older UE always sends 1
+		constexpr bool bIsValidProvider = true;
 		const double MachineTimeOffset = LiveLinkMessageBusHelper::CalculateProviderMachineOffset(Message.CreationPlatformTime, Context);
 		LastProviderPoolResults.Emplace(MakeShared<FProviderPollResult, ESPMode::ThreadSafe>(Context->GetSender(), Message.ProviderName, Message.MachineName, MachineTimeOffset, bIsValidProvider));
 	}
diff --git a/Engine/Source/Runtime/LiveLinkInterface/Public/ILiveLinkClient.h b/Engine/Source/Runtime/LiveLinkInterface/Public/ILiveLinkClient.h
--- a/Engine/Source/Runtime/LiveLinkInterface/Public/ILiveLinkClient.h
+++ b/Engine/Source/Runtime/LiveLinkInterface/Public/ILiveLinkClient.h
@@ -115,7 +115,7 @@ PRAGMA_ENABLE_DEPRECATION_WARNINGS
 class LIVELINKINTERFACE_API ILiveLinkClient : public ILiveLinkClient_Base_DEPRECATED
 {
 public:
-	static const int32 LIVELINK_VERSION = 3;
+	static const int32 LIVELINK_VERSION = 2;
 
 	virtual ~ILiveLinkClient() {}
 
diff --git a/Engine/Source/Runtime/LiveLinkMessageBusFramework/Private/LiveLinkProvider.cpp b/Engine/Source/Runtime/LiveLinkMessageBusFramework/Private/LiveLinkProvider.cpp
--- a/Engine/Source/Runtime/LiveLinkMessageBusFramework/Private/LiveLinkProvider.cpp
+++ b/Engine/Source/Runtime/LiveLinkMessageBusFramework/Private/LiveLinkProvider.cpp
@@ -14,7 +14,7 @@
 
 DEFINE_LOG_CATEGORY_STATIC(LogLiveLinkMessageBus, Warning, All);
 
-static const int32 LIVELINK_SupportedVersion = 3;
+static const int32 LIVELINK_SupportedVersion = 2;
 
 
 FName FLiveLinkMessageAnnotation::SubjectAnnotation = TEXT("SubjectName");
@@ -463,7 +463,7 @@ PRAGMA_ENABLE_DEPRECATION_WARNINGS
 
 void FLiveLinkProvider::HandlePingMessage(const FLiveLinkPingMessage& Message, const TSharedRef<IMessageContext, ESPMode::ThreadSafe>& Context)
 {
-	if (Message.LiveLinkVersion != LIVELINK_SupportedVersion)
+	if (Message.LiveLinkVersion < LIVELINK_SupportedVersion)
 	{
 		UE_LOG(LogLiveLinkMessageBus, Warning, TEXT("A unsupported version of LiveLink is trying to communicate. Requested version: '%d'. Supported version: '%d'."), Message.LiveLinkVersion, LIVELINK_SupportedVersion)
 		return;
@@ -476,7 +476,7 @@ void FLiveLinkProvider::HandleConnectMessage(const FLiveLinkConnectMessage& Mess
 {
 	FScopeLock Lock(&CriticalSection);
 
-	if (Message.LiveLinkVersion != LIVELINK_SupportedVersion)
+	if (Message.LiveLinkVersion < LIVELINK_SupportedVersion)
 	{
 		UE_LOG(LogLiveLinkMessageBus, Error, TEXT("A unsupported version of LiveLink is trying to connect. Requested version: '%d'. Supported version: '%d'."), Message.LiveLinkVersion, LIVELINK_SupportedVersion)
 		return;
