PR #8996: Live Link API changes to allow deriving from LiveLinkProvider and MessageBusSource classes (Contributed by [REDACTED_PERSON])

diff --git a/Engine/Plugins/Animation/LiveLink/Source/LiveLink/Private/LiveLinkMessageBusSource.cpp b/Engine/Plugins/Animation/LiveLink/Source/LiveLink/Private/LiveLinkMessageBusSource.cpp
--- a/Engine/Plugins/Animation/LiveLink/Source/LiveLink/Private/LiveLinkMessageBusSource.cpp
+++ b/Engine/Plugins/Animation/LiveLink/Source/LiveLink/Private/LiveLinkMessageBusSource.cpp
@@ -8,6 +8,7 @@
 #include "LiveLinkHeartbeatEmitter.h"
 #include "LiveLinkLog.h"
 #include "LiveLinkMessageBusDiscoveryManager.h"
+#include "LiveLinkMessageBusSourceSettings.h"
 #include "LiveLinkMessages.h"
 #include "LiveLinkRoleTrait.h"
 #include "LiveLinkSettings.h"
@@ -16,7 +17,6 @@
 #include "MessageEndpointBuilder.h"
 #include "Misc/App.h"
 
-
 FLiveLinkMessageBusSource::FLiveLinkMessageBusSource(const FText& InSourceType, const FText& InSourceMachineName, const FMessageAddress& InConnectionAddress, double InMachineTimeOffset)
 	: ConnectionAddress(InConnectionAddress)
 	, SourceType(InSourceType)
@@ -41,15 +41,7 @@ void FLiveLinkMessageBusSource::ReceiveClient(ILiveLinkClient* InClient, FGuid I
 		RoleInstances.Add(RoleClass->GetDefaultObject<ULiveLinkRole>());
 	}
 
-	PRAGMA_DISABLE_DEPRECATION_WARNINGS
-	MessageEndpoint = FMessageEndpoint::Builder(TEXT("LiveLinkMessageBusSource"))
-		.Handling<FLiveLinkSubjectDataMessage>(this, &FLiveLinkMessageBusSource::HandleSubjectData)
-		.Handling<FLiveLinkSubjectFrameMessage>(this, &FLiveLinkMessageBusSource::HandleSubjectFrame)
-		.Handling<FLiveLinkHeartbeatMessage>(this, &FLiveLinkMessageBusSource::HandleHeartbeat)
-		.Handling<FLiveLinkClearSubject>(this, &FLiveLinkMessageBusSource::HandleClearSubject)
-		.ReceivingOnAnyThread()
-		.WithCatchall(this, &FLiveLinkMessageBusSource::InternalHandleMessage);
-	PRAGMA_ENABLE_DEPRECATION_WARNINGS
+	MessageEndpoint = CreateAndInitializeMessageEndpoint();
 
 	if (ConnectionAddress.IsValid())
 	{
@@ -191,21 +183,11 @@ void FLiveLinkMessageBusSource::InternalHandleMessage(const TSharedRef<IMessageC
 	const FLiveLinkSubjectKey SubjectKey(SourceGuid, SubjectName);
 	if (bIsStaticData)
 	{
-		check(MessageTypeInfo->IsChildOf(FLiveLinkBaseStaticData::StaticStruct()));
-
-		FLiveLinkStaticDataStruct DataStruct(MessageTypeInfo);
-		DataStruct.InitializeWith(MessageTypeInfo, reinterpret_cast<const FLiveLinkBaseStaticData*>(Context->GetMessage()));
-		Client->PushSubjectStaticData_AnyThread(SubjectKey, SubjectRole, MoveTemp(DataStruct));
+		InitializeAndPushStaticData_AnyThread(SubjectName, SubjectRole, SubjectKey, Context, MessageTypeInfo);
 	}
 	else
 	{
-		check(MessageTypeInfo->IsChildOf(FLiveLinkBaseFrameData::StaticStruct()));
-
-		FLiveLinkFrameDataStruct DataStruct(MessageTypeInfo);
-		const FLiveLinkBaseFrameData* Message = reinterpret_cast<const FLiveLinkBaseFrameData*>(Context->GetMessage());
-		DataStruct.InitializeWith(MessageTypeInfo, Message);
-		DataStruct.GetBaseData()->WorldTime = Message->WorldTime.GetOffsettedTime();
-		Client->PushSubjectFrameData_AnyThread(SubjectKey, MoveTemp(DataStruct));
+		InitializeAndPushFrameData_AnyThread(SubjectName, SubjectKey, Context, MessageTypeInfo);
 	}
 }
 
@@ -227,6 +209,58 @@ FText FLiveLinkMessageBusSource::GetSourceStatus() const
 	return NSLOCTEXT("LiveLinkMessageBusSource", "TimeoutStatus", "Not responding");
 }
 
+TSubclassOf<ULiveLinkSourceSettings> FLiveLinkMessageBusSource::GetSettingsClass() const
+{
+	return ULiveLinkMessageBusSourceSettings::StaticClass();
+}
+
+TSharedPtr<FMessageEndpoint, ESPMode::ThreadSafe> FLiveLinkMessageBusSource::CreateAndInitializeMessageEndpoint()
+{
+	FMessageEndpointBuilder EndpointBuilder = FMessageEndpoint::Builder(GetSourceName());
+	InitializeMessageEndpoint(EndpointBuilder);
+	return EndpointBuilder;
+}
+
+void FLiveLinkMessageBusSource::InitializeAndPushStaticData_AnyThread(FName SubjectName,
+																	  TSubclassOf<ULiveLinkRole> SubjectRole,
+																	  const FLiveLinkSubjectKey& SubjectKey,
+																	  const TSharedRef<IMessageContext, ESPMode::ThreadSafe>& Context,
+																	  UScriptStruct* MessageTypeInfo)
+{
+	check(MessageTypeInfo->IsChildOf(FLiveLinkBaseStaticData::StaticStruct()));
+
+	FLiveLinkStaticDataStruct DataStruct(MessageTypeInfo);
+	DataStruct.InitializeWith(MessageTypeInfo, reinterpret_cast<const FLiveLinkBaseStaticData*>(Context->GetMessage()));
+	PushClientSubjectStaticData_AnyThread(SubjectKey, SubjectRole, MoveTemp(DataStruct));
+}
+
+void FLiveLinkMessageBusSource::InitializeAndPushFrameData_AnyThread(FName SubjectName,
+																	 const FLiveLinkSubjectKey& SubjectKey,
+																	 const TSharedRef<IMessageContext, ESPMode::ThreadSafe>& Context,
+																	 UScriptStruct* MessageTypeInfo)
+{
+	check(MessageTypeInfo->IsChildOf(FLiveLinkBaseFrameData::StaticStruct()));
+
+	FLiveLinkFrameDataStruct DataStruct(MessageTypeInfo);
+	const FLiveLinkBaseFrameData* Message = reinterpret_cast<const FLiveLinkBaseFrameData*>(Context->GetMessage());
+	DataStruct.InitializeWith(MessageTypeInfo, Message);
+	DataStruct.GetBaseData()->WorldTime = Message->WorldTime.GetOffsettedTime();
+	PushClientSubjectFrameData_AnyThread(SubjectKey, MoveTemp(DataStruct));
+}
+
+void FLiveLinkMessageBusSource::PushClientSubjectStaticData_AnyThread(const FLiveLinkSubjectKey& SubjectKey,
+																	  TSubclassOf<ULiveLinkRole> Role,
+																	  FLiveLinkStaticDataStruct&& StaticData)
+{
+	Client->PushSubjectStaticData_AnyThread(SubjectKey, Role, MoveTemp(StaticData));
+}
+
+void FLiveLinkMessageBusSource::PushClientSubjectFrameData_AnyThread(const FLiveLinkSubjectKey& SubjectKey,
+																	 FLiveLinkFrameDataStruct&& FrameData)
+{
+	Client->PushSubjectFrameData_AnyThread(SubjectKey, MoveTemp(FrameData));
+}
+
 void FLiveLinkMessageBusSource::HandleHeartbeat(const FLiveLinkHeartbeatMessage& Message, const TSharedRef<IMessageContext, ESPMode::ThreadSafe>& Context)
 {
 	UpdateConnectionLastActive();
@@ -254,7 +288,7 @@ void FLiveLinkMessageBusSource::SendConnectMessage()
 {
 	FLiveLinkConnectMessage* ConnectMessage = FMessageEndpoint::MakeMessage<FLiveLinkConnectMessage>();
 	ConnectMessage->LiveLinkVersion = ILiveLinkClient::LIVELINK_VERSION;
-	MessageEndpoint->Send(ConnectMessage, ConnectionAddress);
+	SendMessage(ConnectMessage);
 	FLiveLinkHeartbeatEmitter& HeartbeatEmitter = ILiveLinkModule::Get().GetHeartbeatEmitter();
 	HeartbeatEmitter.StartHeartbeat(ConnectionAddress, MessageEndpoint);
 	bIsValid = true;
@@ -281,6 +315,25 @@ bool FLiveLinkMessageBusSource::RequestSourceShutdown()
 	return true;
 }
 
+const FName& FLiveLinkMessageBusSource::GetSourceName() const
+{
+	static FName Name(TEXT("LiveLinkMessageBusSource"));
+	return Name;
+}
+
+void FLiveLinkMessageBusSource::InitializeMessageEndpoint(FMessageEndpointBuilder& EndpointBuilder)
+{
+PRAGMA_DISABLE_DEPRECATION_WARNINGS
+	MessageEndpoint = EndpointBuilder
+	.Handling<FLiveLinkSubjectDataMessage>(this, &FLiveLinkMessageBusSource::HandleSubjectData)
+	.Handling<FLiveLinkSubjectFrameMessage>(this, &FLiveLinkMessageBusSource::HandleSubjectFrame)
+	.Handling<FLiveLinkHeartbeatMessage>(this, &FLiveLinkMessageBusSource::HandleHeartbeat)
+	.Handling<FLiveLinkClearSubject>(this, &FLiveLinkMessageBusSource::HandleClearSubject)
+	.ReceivingOnAnyThread()
+	.WithCatchall(this, &FLiveLinkMessageBusSource::InternalHandleMessage);
+PRAGMA_ENABLE_DEPRECATION_WARNINGS
+}
+
 PRAGMA_DISABLE_DEPRECATION_WARNINGS
 void FLiveLinkMessageBusSource::HandleSubjectData(const FLiveLinkSubjectDataMessage& Message, const TSharedRef<IMessageContext, ESPMode::ThreadSafe>& Context)
 {
diff --git a/Engine/Plugins/Animation/LiveLink/Source/LiveLink/Private/LiveLinkMessageBusSource.h b/Engine/Plugins/Animation/LiveLink/Source/LiveLink/Private/LiveLinkMessageBusSource.h
deleted file mode 100644
diff --git a/Engine/Plugins/Animation/LiveLink/Source/LiveLink/Private/LiveLinkMessageBusSourceFactory.cpp b/Engine/Plugins/Animation/LiveLink/Source/LiveLink/Private/LiveLinkMessageBusSourceFactory.cpp
--- a/Engine/Plugins/Animation/LiveLink/Source/LiveLink/Private/LiveLinkMessageBusSourceFactory.cpp
+++ b/Engine/Plugins/Animation/LiveLink/Source/LiveLink/Private/LiveLinkMessageBusSourceFactory.cpp
@@ -5,6 +5,8 @@
 #include "LiveLinkMessageBusSource.h"
 #include "LiveLinkMessageBusFinder.h"
 
+#include "SLiveLinkMessageBusSourceFactory.h"
+
 #include "Features/IModularFeatures.h"
 #include "Misc/MessageDialog.h"
 
@@ -26,6 +28,14 @@ TSharedPtr<SWidget> ULiveLinkMessageBusSourceFactory::BuildCreationPanel(FOnLive
 		.OnSourceSelected(FOnLiveLinkMessageBusSourceSelected::CreateUObject(this, &ULiveLinkMessageBusSourceFactory::OnSourceSelected, InOnLiveLinkSourceCreated));
 }
 
+TSharedPtr<FLiveLinkMessageBusSource> ULiveLinkMessageBusSourceFactory::MakeSource(const FText& Name,
+																				   const FText& MachineName,
+																				   const FMessageAddress& Address,
+																				   double TimeOffset) const
+{
+	return MakeShared<FLiveLinkMessageBusSource>(Name, MachineName, Address, TimeOffset);
+}
+
 TSharedPtr<ILiveLinkSource> ULiveLinkMessageBusSourceFactory::CreateSource(const FString& ConnectionString) const
 {
 	FString Name;
@@ -35,7 +45,7 @@ TSharedPtr<ILiveLinkSource> ULiveLinkMessageBusSourceFactory::CreateSource(const
 	}
 
 	const double TimeOffset = 0.0;
-	return MakeShared<FLiveLinkMessageBusSource>(FText::FromString(Name), FText::GetEmpty(), FMessageAddress(), TimeOffset);
+	return MakeSource(FText::FromString(Name), FText::GetEmpty(), FMessageAddress(), TimeOffset);
 }
 
 
@@ -72,7 +82,7 @@ void ULiveLinkMessageBusSourceFactory::OnSourceSelected(FProviderPollResultPtr S
 		}
 #endif
 
-		TSharedPtr<FLiveLinkMessageBusSource> SharedPtr = MakeShared<FLiveLinkMessageBusSource>(FText::FromString(SelectedSource->Name), FText::FromString(SelectedSource->MachineName), SelectedSource->Address, SelectedSource->MachineTimeOffset);
+		TSharedPtr<FLiveLinkMessageBusSource> SharedPtr = MakeSource(FText::FromString(SelectedSource->Name), FText::FromString(SelectedSource->MachineName), SelectedSource->Address, SelectedSource->MachineTimeOffset);
 		FString ConnectionString = CreateConnectionString(*SelectedSource.Get());
 		InOnLiveLinkSourceCreated.ExecuteIfBound(SharedPtr, MoveTemp(ConnectionString));
 	}
diff --git a/Engine/Plugins/Animation/LiveLink/Source/LiveLink/Private/LiveLinkMessageBusSourceFactory.h b/Engine/Plugins/Animation/LiveLink/Source/LiveLink/Private/LiveLinkMessageBusSourceFactory.h
deleted file mode 100644
diff --git a/Engine/Plugins/Animation/LiveLink/Source/LiveLink/Public/LiveLinkMessageBusSource.h b/Engine/Plugins/Animation/LiveLink/Source/LiveLink/Public/LiveLinkMessageBusSource.h
new file mode 100644
--- /dev/null
+++ b/Engine/Plugins/Animation/LiveLink/Source/LiveLink/Public/LiveLinkMessageBusSource.h
@@ -0,0 +1,122 @@
+// Copyright Epic Games, Inc. All Rights Reserved.
+
+#pragma once
+
+#include "ILiveLinkSource.h"
+
+#include "HAL/ThreadSafeBool.h"
+#include "IMessageContext.h"
+#include "LiveLinkRole.h"
+#include "MessageEndpoint.h"
+
+class ILiveLinkClient;
+struct FLiveLinkPongMessage;
+struct FLiveLinkSubjectDataMessage;
+struct FLiveLinkSubjectFrameMessage;
+struct FLiveLinkHeartbeatMessage;
+struct FLiveLinkClearSubject;
+
+class LIVELINK_API FLiveLinkMessageBusSource : public ILiveLinkSource
+{
+public:
+
+	FLiveLinkMessageBusSource(const FText& InSourceType, const FText& InSourceMachineName, const FMessageAddress& InConnectionAddress, double InMachineTimeOffset);
+
+	//~ Begin ILiveLinkSource interface
+	virtual void InitializeSettings(ULiveLinkSourceSettings* Settings) override;
+	virtual void ReceiveClient(ILiveLinkClient* InClient, FGuid InSourceGuid) override;
+	virtual void Update() override;
+
+	virtual bool IsSourceStillValid() const override;
+
+	virtual bool RequestSourceShutdown() override;
+
+	virtual FText GetSourceType() const override { return SourceType; }
+	virtual FText GetSourceMachineName() const override { return SourceMachineName; }
+	virtual FText GetSourceStatus() const override;
+
+	virtual TSubclassOf<ULiveLinkSourceSettings> GetSettingsClass() const override;
+	//~ End ILiveLinkSource interface
+
+protected:
+	// Returns the source name to uniquely identify it among the FLiveLinkMessageBusSource classes
+	virtual const FName& GetSourceName() const;
+	// This lets child classes the opportunity to add custom message handlers to the endpoint builder
+	virtual void InitializeMessageEndpoint(FMessageEndpointBuilder& EndpointBuilder);
+	bool IsMessageEndpointConnected() const { return ConnectionAddress.IsValid() && MessageEndpoint.IsValid() && MessageEndpoint->IsConnected(); }
+
+	// Initialize the static data and send it to the clients
+	virtual void InitializeAndPushStaticData_AnyThread(FName SubjectName,
+													   TSubclassOf<ULiveLinkRole> SubjectRole,
+													   const FLiveLinkSubjectKey& SubjectKey,
+													   const TSharedRef<IMessageContext, ESPMode::ThreadSafe>& Context,
+													   UScriptStruct* MessageTypeInfo);
+	// Initialize the frame data and send it to the clients
+	virtual void InitializeAndPushFrameData_AnyThread(FName SubjectName,
+													  const FLiveLinkSubjectKey& SubjectKey,
+													  const TSharedRef<IMessageContext, ESPMode::ThreadSafe>& Context,
+													  UScriptStruct* MessageTypeInfo);
+
+	// Send the static data to the clients
+	void PushClientSubjectStaticData_AnyThread(const FLiveLinkSubjectKey& SubjectKey, TSubclassOf<ULiveLinkRole> Role, FLiveLinkStaticDataStruct&& StaticData);
+	// Send the frame data to the clients
+	void PushClientSubjectFrameData_AnyThread(const FLiveLinkSubjectKey& SubjectKey, FLiveLinkFrameDataStruct&& FrameData);
+
+	// Send a message through the endpoint
+	template<typename MessageType>
+	void SendMessage(MessageType* Message)
+	{
+		if (!Message || !IsMessageEndpointConnected())
+		{
+			return;
+		}
+
+		MessageEndpoint->Send(Message, ConnectionAddress);
+	}
+
+private:
+	TSharedPtr<FMessageEndpoint, ESPMode::ThreadSafe> CreateAndInitializeMessageEndpoint();
+
+	//~ Message bus message handlers
+	PRAGMA_DISABLE_DEPRECATION_WARNINGS
+	void HandleSubjectData(const FLiveLinkSubjectDataMessage& Message, const TSharedRef<IMessageContext, ESPMode::ThreadSafe>& Context);
+	void HandleSubjectFrame(const FLiveLinkSubjectFrameMessage& Message, const TSharedRef<IMessageContext, ESPMode::ThreadSafe>& Context);
+	PRAGMA_ENABLE_DEPRECATION_WARNINGS
+
+	void HandleHeartbeat(const FLiveLinkHeartbeatMessage& Message, const TSharedRef<IMessageContext, ESPMode::ThreadSafe>& Context);
+	void HandleClearSubject(const FLiveLinkClearSubject& Message, const TSharedRef<IMessageContext, ESPMode::ThreadSafe>& Context);
+	void InternalHandleMessage(const TSharedRef<IMessageContext, ESPMode::ThreadSafe>& Context);
+	//~ End Message bus message handlers
+
+	// Threadsafe update of the last active time
+	FORCEINLINE void UpdateConnectionLastActive();
+
+	void SendConnectMessage();
+
+	ILiveLinkClient* Client;
+
+	// Our identifier in LiveLink
+	FGuid SourceGuid;
+
+	// List of the roles available when the bus was opened
+	TArray<TWeakObjectPtr<ULiveLinkRole>> RoleInstances;
+
+	TSharedPtr<FMessageEndpoint, ESPMode::ThreadSafe> MessageEndpoint;
+
+	FMessageAddress ConnectionAddress;
+
+	FText SourceType;
+	FText SourceMachineName;
+
+	// Time we last received anything 
+	double ConnectionLastActive;
+
+	// Current Validity of Source
+	FThreadSafeBool bIsValid;
+
+	// Critical section to allow for threadsafe updating of the connection time
+	FCriticalSection ConnectionLastActiveSection;
+
+	// Offset between sender's machine engine time and receiver's machine engine time
+	double MachineTimeOffset;
+};
\ No newline at end of file
diff --git a/Engine/Plugins/Animation/LiveLink/Source/LiveLink/Public/LiveLinkMessageBusSourceFactory.h b/Engine/Plugins/Animation/LiveLink/Source/LiveLink/Public/LiveLinkMessageBusSourceFactory.h
new file mode 100644
--- /dev/null
+++ b/Engine/Plugins/Animation/LiveLink/Source/LiveLink/Public/LiveLinkMessageBusSourceFactory.h
@@ -0,0 +1,34 @@
+// Copyright Epic Games, Inc. All Rights Reserved.
+
+#pragma once
+
+#include "LiveLinkMessageBusFinder.h"		// For FProviderPollResultPtr typedef
+#include "LiveLinkSourceFactory.h"
+#include "LiveLinkMessageBusSourceFactory.generated.h"
+
+
+UCLASS()
+class LIVELINK_API ULiveLinkMessageBusSourceFactory : public ULiveLinkSourceFactory
+{
+public:
+	GENERATED_BODY()
+
+	virtual FText GetSourceDisplayName() const;
+	virtual FText GetSourceTooltip() const;
+
+	virtual EMenuType GetMenuType() const override { return EMenuType::SubPanel; }
+	virtual TSharedPtr<SWidget> BuildCreationPanel(FOnLiveLinkSourceCreated OnLiveLinkSourceCreated) const override;
+	virtual TSharedPtr<ILiveLinkSource> CreateSource(const FString& ConnectionString) const override;
+
+	static FString CreateConnectionString(const struct FProviderPollResult& Result);
+
+protected:
+	// This allows the creation of a message bus source derived from FLiveLinkMessageBusSource
+	virtual TSharedPtr<class FLiveLinkMessageBusSource> MakeSource(const FText& Name,
+																   const FText& MachineName,
+																   const FMessageAddress& Address,
+																   double TimeOffset) const;
+
+private:
+	void OnSourceSelected(FProviderPollResultPtr SelectedSource, FOnLiveLinkSourceCreated OnLiveLinkSourceCreated) const;
+};
diff --git a/Engine/Source/Runtime/LiveLinkMessageBusFramework/Private/LiveLinkProvider.cpp b/Engine/Source/Runtime/LiveLinkMessageBusFramework/Private/LiveLinkProvider.cpp
--- a/Engine/Source/Runtime/LiveLinkMessageBusFramework/Private/LiveLinkProvider.cpp
+++ b/Engine/Source/Runtime/LiveLinkMessageBusFramework/Private/LiveLinkProvider.cpp
@@ -1,6 +1,7 @@
 // Copyright Epic Games, Inc. All Rights Reserved.
 
 #include "LiveLinkProvider.h"
+#include "LiveLinkProviderImpl.h"
 
 #include "HAL/PlatformProcess.h"
 #include "IMessageContext.h"
@@ -16,6 +17,8 @@ DEFINE_LOG_CATEGORY_STATIC(LogLiveLinkMessageBus, Warning, All);
 
 static const int32 LIVELINK_SupportedVersion = 2;
 
+template TSharedPtr<ILiveLinkProvider> ILiveLinkProvider::CreateLiveLinkProvider<FLiveLinkProvider>(const FString&,
+																									FMessageEndpointBuilder&&);
 
 FName FLiveLinkMessageAnnotation::SubjectAnnotation = TEXT("SubjectName");
 FName FLiveLinkMessageAnnotation::RoleAnnotation = TEXT("Role");
@@ -53,28 +56,6 @@ private:
 
 const double FConnectionValidator::CONNECTION_TIMEOUT = 10.f;
 
-PRAGMA_DISABLE_DEPRECATION_WARNINGS
-// Subject that the application has told us about
-struct FTrackedSubject
-{
-	// Ref skeleton to go with transform data
-	FLiveLinkRefSkeleton RefSkeleton;
-
-	// Bone transform data
-	TArray<FTransform> Transforms;
-
-	// Curve data
-	TArray<FLiveLinkCurveElement> Curves;
-
-	// MetaData for subject
-	FLiveLinkMetaData MetaData;
-
-	// Incrementing time (application time) for interpolation purposes
-	double Time;
-};
-PRAGMA_ENABLE_DEPRECATION_WARNINGS
-
-
 // Static Subject data that the application has told us about
 struct FTrackedStaticData
 {
@@ -106,359 +87,313 @@ struct FTrackedFrameData
 };
 
 PRAGMA_DISABLE_DEPRECATION_WARNINGS
-struct FLiveLinkProvider : public ILiveLinkProvider
+
+// Validate our current connections
+void FLiveLinkProvider::ValidateConnections()
 {
-private:
-	const FString ProviderName;
-	const FString MachineName;
+	FConnectionValidator Validator;
 
-	TSharedPtr<FMessageEndpoint, ESPMode::ThreadSafe> MessageEndpoint;
+	const int32 RemovedConnections = ConnectedAddresses.RemoveAll([=](const FTrackedAddress& Address) { return !Validator(Address); });
 
-	/** Lock to stop multiple threads accessing the CurrentPreset at the same time */
-	mutable FCriticalSection CriticalSection;
+	if (RemovedConnections > 0)
+	{
+		OnConnectionStatusChanged.Broadcast();
+	}
+}
 
-	// Array of our current connections
-	TArray<FTrackedAddress> ConnectedAddresses;
+// Get the cached data for the named subject
+FTrackedSubject& FLiveLinkProvider::GetTrackedSubject(const FName& SubjectName)
+{
+	return Subjects.FindOrAdd(SubjectName);
+}
 
-	// Cache of our current subject state
-	TArray<FTrackedStaticData> StaticDatas;
-	TArray<FTrackedFrameData> FrameDatas;
-	TMap<FName, FTrackedSubject> Subjects;
+// Send hierarchy data for named subject
+void FLiveLinkProvider::SendSubject(FName SubjectName, const FTrackedSubject& Subject)
+{
+	FLiveLinkSubjectDataMessage* SubjectData = FMessageEndpoint::MakeMessage<FLiveLinkSubjectDataMessage>();
+	SubjectData->RefSkeleton = Subject.RefSkeleton;
+	SubjectData->SubjectName = SubjectName;
 
-	// Delegate to notify interested parties when the client sources have changed
-	FLiveLinkProviderConnectionStatusChanged OnConnectionStatusChanged;
-	
-	//Message bus message handlers
-	void HandlePingMessage(const FLiveLinkPingMessage& Message, const TSharedRef<IMessageContext, ESPMode::ThreadSafe>& Context);
-	void HandleConnectMessage(const FLiveLinkConnectMessage& Message, const TSharedRef<IMessageContext, ESPMode::ThreadSafe>& Context);
-	void HandleHeartbeat(const FLiveLinkHeartbeatMessage& Message, const TSharedRef<IMessageContext, ESPMode::ThreadSafe>& Context);
-	// End message bus message handlers
+	TArray<FMessageAddress> Addresses;
+	GetConnectedAddresses(Addresses);
 
-	// Validate our current connections
-	void ValidateConnections()
-	{
-		FConnectionValidator Validator;
+	MessageEndpoint->Send(SubjectData, Addresses);
+}
 
-		const int32 RemovedConnections = ConnectedAddresses.RemoveAll([=](const FTrackedAddress& Address) { return !Validator(Address); });
+// Send frame data for named subject
+void FLiveLinkProvider::SendSubjectFrame(FName SubjectName, const FTrackedSubject& Subject)
+{
+	FLiveLinkSubjectFrameMessage* SubjectFrame = FMessageEndpoint::MakeMessage<FLiveLinkSubjectFrameMessage>();
+	SubjectFrame->Transforms = Subject.Transforms;
+	SubjectFrame->SubjectName = SubjectName;
+	SubjectFrame->Curves = Subject.Curves;
+	SubjectFrame->MetaData = Subject.MetaData;
+	SubjectFrame->Time = Subject.Time;
 
-		if (RemovedConnections > 0)
-		{
-			OnConnectionStatusChanged.Broadcast();
-		}
-	}
+	TArray<FMessageAddress> Addresses;
+	GetConnectedAddresses(Addresses);
 
-	// Get the cached data for the named subject
-	FTrackedSubject& GetTrackedSubject(const FName& SubjectName)
-	{
-		return Subjects.FindOrAdd(SubjectName);
-	}
+	MessageEndpoint->Send(SubjectFrame, Addresses);
+}
 
-	// Send hierarchy data for named subject
-	void SendSubject(FName SubjectName, const FTrackedSubject& Subject)
-	{
-		FLiveLinkSubjectDataMessage* SubjectData = FMessageEndpoint::MakeMessage<FLiveLinkSubjectDataMessage>();
-		SubjectData->RefSkeleton = Subject.RefSkeleton;
-		SubjectData->SubjectName = SubjectName;
+// Get the cached data for the named subject
+FTrackedStaticData* FLiveLinkProvider::GetLastSubjectStaticData(const FName& SubjectName)
+{
+	return StaticDatas.FindByKey(SubjectName);
+}
 
-		ValidateConnections();
-		TArray<FMessageAddress> Addresses;
-		Addresses.Reserve(ConnectedAddresses.Num());
-		for (const FTrackedAddress& Address : ConnectedAddresses)
-		{
-			Addresses.Add(Address.Address);
-		}
+FTrackedFrameData* FLiveLinkProvider::GetLastSubjectFrameData(const FName& SubjectName)
+{
+	return FrameDatas.FindByKey(SubjectName);
+}
 
-		MessageEndpoint->Send(SubjectData, Addresses);
+void FLiveLinkProvider::SetLastSubjectStaticData(FName SubjectName, TSubclassOf<ULiveLinkRole> Role, FLiveLinkStaticDataStruct&& StaticData)
+{
+	FTrackedStaticData* Result = StaticDatas.FindByKey(SubjectName);
+	if (Result)
+	{
+		Result->StaticData = MoveTemp(StaticData);
+		Result->RoleClass = Role.Get();
 	}
-
-	// Send frame data for named subject
-	void SendSubjectFrame(FName SubjectName, const FTrackedSubject& Subject)
+	else
 	{
-		FLiveLinkSubjectFrameMessage* SubjectFrame = FMessageEndpoint::MakeMessage<FLiveLinkSubjectFrameMessage>();
-		SubjectFrame->Transforms = Subject.Transforms;
-		SubjectFrame->SubjectName = SubjectName;
-		SubjectFrame->Curves = Subject.Curves;
-		SubjectFrame->MetaData = Subject.MetaData;
-		SubjectFrame->Time = Subject.Time;
-
-		ValidateConnections();
-
-		TArray<FMessageAddress> Addresses;
-		Addresses.Reserve(ConnectedAddresses.Num());
-		for (const FTrackedAddress& Address : ConnectedAddresses)
-		{
-			Addresses.Add(Address.Address);
-		}
-
-		MessageEndpoint->Send(SubjectFrame, Addresses);
+		StaticDatas.Emplace(SubjectName, Role.Get(), MoveTemp(StaticData));
 	}
+}
 
-	// Get the cached data for the named subject
-	FTrackedStaticData* GetLastSubjectStaticData(const FName& SubjectName)
+void FLiveLinkProvider::SetLastSubjectFrameData(FName SubjectName, FLiveLinkFrameDataStruct&& FrameData)
+{
+	FTrackedFrameData* Result = FrameDatas.FindByKey(SubjectName);
+	if (Result)
 	{
-		return StaticDatas.FindByKey(SubjectName);
+		Result->FrameData = MoveTemp(FrameData);
 	}
-
-	FTrackedFrameData* GetLastSubjectFrameData(const FName& SubjectName)
+	else
 	{
-		return FrameDatas.FindByKey(SubjectName);
+		FrameDatas.Emplace(SubjectName, MoveTemp(FrameData));
 	}
+}
 
-	void SetLastSubjectStaticData(FName SubjectName, TSubclassOf<ULiveLinkRole> Role, FLiveLinkStaticDataStruct&& StaticData)
+// Clear a existing track subject
+void FLiveLinkProvider::ClearTrackedSubject(const FName& SubjectName)
+{
+	Subjects.Remove(SubjectName);
+	const int32 FrameIndex = FrameDatas.IndexOfByKey(SubjectName);
+	if (FrameIndex != INDEX_NONE)
 	{
-		FTrackedStaticData* Result = StaticDatas.FindByKey(SubjectName);
-		if (Result)
-		{
-			Result->StaticData = MoveTemp(StaticData);
-			Result->RoleClass = Role.Get();
-		}
-		else
-		{
-			StaticDatas.Emplace(SubjectName, Role.Get(), MoveTemp(StaticData));
-		}
+		FrameDatas.RemoveAtSwap(FrameIndex);
 	}
-
-	void SetLastSubjectFrameData(FName SubjectName, FLiveLinkFrameDataStruct&& FrameData)
+	const int32 StaticIndex = StaticDatas.IndexOfByKey(SubjectName);
+	if (StaticIndex != INDEX_NONE)
 	{
-		FTrackedFrameData* Result = FrameDatas.FindByKey(SubjectName);
-		if (Result)
-		{
-			Result->FrameData = MoveTemp(FrameData);
-		}
-		else
-		{
-			FrameDatas.Emplace(SubjectName, MoveTemp(FrameData));
-		}
+		StaticDatas.RemoveAtSwap(StaticIndex);
 	}
+}
+
+void FLiveLinkProvider::SendClearSubjectToConnections(FName SubjectName)
+{
+	TArray<FMessageAddress> MessageAddresses;
+	GetConnectedAddresses(MessageAddresses);
+
+	MessageEndpoint->Send(FMessageEndpoint::MakeMessage<FLiveLinkClearSubject>(SubjectName), EMessageFlags::Reliable, nullptr, MessageAddresses, FTimespan::Zero(), FDateTime::MaxValue());
+}
+
+FLiveLinkProvider::FLiveLinkProvider(const FString& InProviderName)
+	: ProviderName(InProviderName)
+	, MachineName(FPlatformProcess::ComputerName())
+{
+	FMessageEndpointBuilder EndpointBuilder = FMessageEndpoint::Builder(*InProviderName);
+	CreateMessageEndpoint(EndpointBuilder);
+}
 
-	// Clear a existing track subject
-	void ClearTrackedSubject(const FName& SubjectName)
+FLiveLinkProvider::FLiveLinkProvider(const FString& InProviderName, FMessageEndpointBuilder&& EndpointBuilder)
+	: ProviderName(InProviderName)
+	, MachineName(FPlatformProcess::ComputerName())
+{
+	CreateMessageEndpoint(EndpointBuilder);
+}
+
+FLiveLinkProvider::~FLiveLinkProvider()
+{
+	if (MessageEndpoint.IsValid())
 	{
-		Subjects.Remove(SubjectName);
-		const int32 FrameIndex = FrameDatas.IndexOfByKey(SubjectName);
-		if (FrameIndex != INDEX_NONE)
-		{
-			FrameDatas.RemoveAtSwap(FrameIndex);
-		}
-		const int32 StaticIndex = StaticDatas.IndexOfByKey(SubjectName);
-		if (FrameIndex != INDEX_NONE)
-		{
-			StaticDatas.RemoveAtSwap(StaticIndex);
-		}
+		// Disable the Endpoint message handling since the message could keep it alive a bit.
+		MessageEndpoint->Disable();
+		MessageEndpoint.Reset();
 	}
+}
 
-	void SendClearSubjectToConnections(FName SubjectName)
-	{
-		ValidateConnections();
+void FLiveLinkProvider::UpdateSubject(const FName& SubjectName, const TArray<FName>& BoneNames, const TArray<int32>& BoneParents)
+{
+	FScopeLock Lock(&CriticalSection);
 
-		TArray<FMessageAddress> MessageAddresses;
-		MessageAddresses.Reserve(ConnectedAddresses.Num());
-		for (const FTrackedAddress& Address : ConnectedAddresses)
-		{
-			MessageAddresses.Add(Address.Address);
-		}
+	FTrackedSubject& Subject = GetTrackedSubject(SubjectName);
+	Subject.RefSkeleton.SetBoneNames(BoneNames);
+	Subject.RefSkeleton.SetBoneParents(BoneParents);
+	Subject.Transforms.Empty();
 
-		MessageEndpoint->Send(FMessageEndpoint::MakeMessage<FLiveLinkClearSubject>(SubjectName), EMessageFlags::Reliable, nullptr, MessageAddresses, FTimespan::Zero(), FDateTime::MaxValue());
-	}
+	SendSubject(SubjectName, Subject);
+}
 
-public:
-	FLiveLinkProvider(const FString& InProviderName)
-		: ProviderName(InProviderName)
-		, MachineName(FPlatformProcess::ComputerName())
-	{
-		MessageEndpoint = FMessageEndpoint::Builder(*InProviderName)
-			.ReceivingOnAnyThread()
-			.Handling<FLiveLinkPingMessage>(this, &FLiveLinkProvider::HandlePingMessage)
-			.Handling<FLiveLinkConnectMessage>(this, &FLiveLinkProvider::HandleConnectMessage)
-			.Handling<FLiveLinkHeartbeatMessage>(this, &FLiveLinkProvider::HandleHeartbeat);
+bool FLiveLinkProvider::UpdateSubjectStaticData(const FName SubjectName, TSubclassOf<ULiveLinkRole> Role, FLiveLinkStaticDataStruct&& StaticData)
+{
+	FScopeLock Lock(&CriticalSection);
 
-		if (MessageEndpoint.IsValid())
-		{
-			MessageEndpoint->Subscribe<FLiveLinkPingMessage>();
-		}
+	if (SubjectName == NAME_None || Role.Get() == nullptr)
+	{
+		return false;
 	}
 
-	virtual ~FLiveLinkProvider()
+	if (Role->GetDefaultObject<ULiveLinkRole>()->GetStaticDataStruct() != StaticData.GetStruct())
 	{
-		if (MessageEndpoint.IsValid())
-		{
-			// Disable the Endpoint message handling since the message could keep it alive a bit.
-			MessageEndpoint->Disable();
-			MessageEndpoint.Reset();
-		}
+		return false;
 	}
 
-	virtual void UpdateSubject(const FName& SubjectName, const TArray<FName>& BoneNames, const TArray<int32>& BoneParents) override
+	if (GetLastSubjectStaticData(SubjectName) != nullptr)
 	{
-		FScopeLock Lock(&CriticalSection);
-
-		FTrackedSubject& Subject = GetTrackedSubject(SubjectName);
-		Subject.RefSkeleton.SetBoneNames(BoneNames);
-		Subject.RefSkeleton.SetBoneParents(BoneParents);
-		Subject.Transforms.Empty();
-
-		SendSubject(SubjectName, Subject);
+		ClearSubject(SubjectName);
 	}
 
-	virtual bool UpdateSubjectStaticData(const FName SubjectName, TSubclassOf<ULiveLinkRole> Role, FLiveLinkStaticDataStruct&& StaticData) override
-	{
-		FScopeLock Lock(&CriticalSection);
+	ValidateConnections();
 
-		if (SubjectName == NAME_None || Role.Get() == nullptr)
+	if (ConnectedAddresses.Num() > 0)
+	{
+		TArray<FMessageAddress> Addresses;
+		Addresses.Reserve(ConnectedAddresses.Num());
+		for (const FTrackedAddress& Address : ConnectedAddresses)
 		{
-			return false;
+			Addresses.Add(Address.Address);
 		}
 
-		if (Role->GetDefaultObject<ULiveLinkRole>()->GetStaticDataStruct() != StaticData.GetStruct())
-		{
-			return false;
-		}
+		TMap<FName, FString> Annotations;
+		Annotations.Add(FLiveLinkMessageAnnotation::SubjectAnnotation, SubjectName.ToString());
+		Annotations.Add(FLiveLinkMessageAnnotation::RoleAnnotation, Role->GetName());
 
-		if (GetLastSubjectStaticData(SubjectName) != nullptr)
-		{
-			ClearSubject(SubjectName);
-		}
+		MessageEndpoint->Send(StaticData.CloneData(), const_cast<UScriptStruct*>(StaticData.GetStruct()), EMessageFlags::Reliable, Annotations, nullptr, Addresses, FTimespan::Zero(), FDateTime::MaxValue());
+	}
 
-		ValidateConnections();
+	SetLastSubjectStaticData(SubjectName, Role, MoveTemp(StaticData));
 
-		if (ConnectedAddresses.Num() > 0)
-		{
-			TArray<FMessageAddress> Addresses;
-			Addresses.Reserve(ConnectedAddresses.Num());
-			for (const FTrackedAddress& Address : ConnectedAddresses)
-			{
-				Addresses.Add(Address.Address);
-			}
-
-			TMap<FName, FString> Annotations;
-			Annotations.Add(FLiveLinkMessageAnnotation::SubjectAnnotation, SubjectName.ToString());
-			Annotations.Add(FLiveLinkMessageAnnotation::RoleAnnotation, Role->GetName());
-
-			MessageEndpoint->Send(StaticData.CloneData(), const_cast<UScriptStruct*>(StaticData.GetStruct()), EMessageFlags::Reliable, Annotations, nullptr, Addresses, FTimespan::Zero(), FDateTime::MaxValue());
-		}
+	return true;
+}
 
-		SetLastSubjectStaticData(SubjectName, Role, MoveTemp(StaticData));
+void FLiveLinkProvider::ClearSubject(const FName& SubjectName)
+{
+	FScopeLock Lock(&CriticalSection);
 
-		return true;
-	}
+	RemoveSubject(SubjectName);
+}
 
-	virtual void ClearSubject(const FName& SubjectName)
-	{
-		FScopeLock Lock(&CriticalSection);
+void FLiveLinkProvider::RemoveSubject(const FName SubjectName)
+{
+	FScopeLock Lock(&CriticalSection);
 
-		RemoveSubject(SubjectName);
-	}
+	ClearTrackedSubject(SubjectName);
+	SendClearSubjectToConnections(SubjectName);
+}
 
-	virtual void RemoveSubject(const FName SubjectName) override
-	{
-		FScopeLock Lock(&CriticalSection);
+void FLiveLinkProvider::UpdateSubjectFrame(const FName& SubjectName, const TArray<FTransform>& BoneTransforms, const TArray<FLiveLinkCurveElement>& CurveData, double Time)
+{
+	FScopeLock Lock(&CriticalSection);
 
-		ClearTrackedSubject(SubjectName);
-		SendClearSubjectToConnections(SubjectName);
-	}
+	FTrackedSubject& Subject = GetTrackedSubject(SubjectName);
 
-	virtual void UpdateSubjectFrame(const FName& SubjectName, const TArray<FTransform>& BoneTransforms, const TArray<FLiveLinkCurveElement>& CurveData, double Time) override
-	{
-		FScopeLock Lock(&CriticalSection);
+	Subject.Transforms = BoneTransforms;
+	Subject.Curves = CurveData;
+	Subject.Time = Time;
 
-		FTrackedSubject& Subject = GetTrackedSubject(SubjectName);
+	SendSubjectFrame(SubjectName, Subject);
+}
 
-		Subject.Transforms = BoneTransforms;
-		Subject.Curves = CurveData;
-		Subject.Time = Time;
+void FLiveLinkProvider::UpdateSubjectFrame(const FName& SubjectName, const TArray<FTransform>& BoneTransforms, const TArray<FLiveLinkCurveElement>& CurveData,
+	const FLiveLinkMetaData& MetaData, double Time)
+{
+	FScopeLock Lock(&CriticalSection);
 
-		SendSubjectFrame(SubjectName, Subject);
-	}
+	FTrackedSubject& Subject = GetTrackedSubject(SubjectName);
 
-	virtual void UpdateSubjectFrame(const FName& SubjectName, const TArray<FTransform>& BoneTransforms, const TArray<FLiveLinkCurveElement>& CurveData,
-		const FLiveLinkMetaData& MetaData, double Time) override
-	{
-		FScopeLock Lock(&CriticalSection);
+	Subject.Transforms = BoneTransforms;
+	Subject.Curves = CurveData;
+	Subject.MetaData = MetaData;
+	Subject.Time = Time;
 
-		FTrackedSubject& Subject = GetTrackedSubject(SubjectName);
+	SendSubjectFrame(SubjectName, Subject);
+}
 
-		Subject.Transforms = BoneTransforms;
-		Subject.Curves = CurveData;
-		Subject.MetaData = MetaData;
-		Subject.Time = Time;
+bool FLiveLinkProvider::UpdateSubjectFrameData(const FName SubjectName, FLiveLinkFrameDataStruct&& FrameData)
+{
+	FScopeLock Lock(&CriticalSection);
 
-		SendSubjectFrame(SubjectName, Subject);
+	if (SubjectName == NAME_None)
+	{
+		return false;
 	}
 
-	virtual bool UpdateSubjectFrameData(const FName SubjectName, FLiveLinkFrameDataStruct&& FrameData) override
+	FTrackedStaticData* StaticData = GetLastSubjectStaticData(SubjectName);
+	if (StaticData == nullptr)
 	{
-		FScopeLock Lock(&CriticalSection);
+		return false;
+	}
 
-		if (SubjectName == NAME_None)
-		{
-			return false;
-		}
+	UClass* RoleClass = StaticData->RoleClass.Get();
+	if (RoleClass == nullptr)
+	{
+		return false;
+	}
 
-		FTrackedStaticData* StaticData = GetLastSubjectStaticData(SubjectName);
-		if (StaticData == nullptr)
-		{
-			return false;
-		}
+	if (RoleClass->GetDefaultObject<ULiveLinkRole>()->GetFrameDataStruct() != FrameData.GetStruct())
+	{
+		return false;
+	}
 
-		UClass* RoleClass = StaticData->RoleClass.Get();
-		if (RoleClass == nullptr)
-		{
-			return false;
-		}
+	ValidateConnections();
 
-		if (RoleClass->GetDefaultObject<ULiveLinkRole>()->GetFrameDataStruct() != FrameData.GetStruct())
+	if (ConnectedAddresses.Num() > 0)
+	{
+		TArray<FMessageAddress> Addresses;
+		Addresses.Reserve(ConnectedAddresses.Num());
+		for (const FTrackedAddress& Address : ConnectedAddresses)
 		{
-			return false;
+			Addresses.Add(Address.Address);
 		}
 
-		ValidateConnections();
+		TMap<FName, FString> Annotations;
+		Annotations.Add(FLiveLinkMessageAnnotation::SubjectAnnotation, SubjectName.ToString());
 
-		if (ConnectedAddresses.Num() > 0)
-		{
-			TArray<FMessageAddress> Addresses;
-			Addresses.Reserve(ConnectedAddresses.Num());
-			for (const FTrackedAddress& Address : ConnectedAddresses)
-			{
-				Addresses.Add(Address.Address);
-			}
+		MessageEndpoint->Send(FrameData.CloneData(), const_cast<UScriptStruct*>(FrameData.GetStruct()), EMessageFlags::None, Annotations, nullptr, Addresses, FTimespan::Zero(), FDateTime::MaxValue());
+	}
 
-			TMap<FName, FString> Annotations;
-			Annotations.Add(FLiveLinkMessageAnnotation::SubjectAnnotation, SubjectName.ToString());
+	SetLastSubjectFrameData(SubjectName, MoveTemp(FrameData));
 
-			MessageEndpoint->Send(FrameData.CloneData(), const_cast<UScriptStruct*>(FrameData.GetStruct()), EMessageFlags::None, Annotations, nullptr, Addresses, FTimespan::Zero(), FDateTime::MaxValue());
-		}
+	return true;
+}
 
-		SetLastSubjectFrameData(SubjectName, MoveTemp(FrameData));
+bool FLiveLinkProvider::HasConnection() const
+{
+	FScopeLock Lock(&CriticalSection);
 
-		return true;
-	}
+	FConnectionValidator Validator;
 
-	virtual bool HasConnection() const
+	for (const FTrackedAddress& Connection : ConnectedAddresses)
 	{
-		FScopeLock Lock(&CriticalSection);
-
-		FConnectionValidator Validator;
-
-		for (const FTrackedAddress& Connection : ConnectedAddresses)
+		if (Validator(Connection))
 		{
-			if (Validator(Connection))
-			{
-				return true;
-			}
+			return true;
 		}
-		return false;
 	}
+	return false;
+}
 
-	virtual FDelegateHandle RegisterConnStatusChangedHandle(const FLiveLinkProviderConnectionStatusChanged::FDelegate& ConnStatusChanged)
-	{
-		return OnConnectionStatusChanged.Add(ConnStatusChanged);
-	}
+FDelegateHandle FLiveLinkProvider::RegisterConnStatusChangedHandle(const FLiveLinkProviderConnectionStatusChanged::FDelegate& ConnStatusChanged)
+{
+	return OnConnectionStatusChanged.Add(ConnStatusChanged);
+}
+
+void FLiveLinkProvider::UnregisterConnStatusChangedHandle(FDelegateHandle Handle)
+{
+	OnConnectionStatusChanged.Remove(Handle);
+}
 
-	virtual void UnregisterConnStatusChangedHandle(FDelegateHandle Handle)
-	{
-		OnConnectionStatusChanged.Remove(Handle);
-	}
-};
 PRAGMA_ENABLE_DEPRECATION_WARNINGS
 
 void FLiveLinkProvider::HandlePingMessage(const FLiveLinkPingMessage& Message, const TSharedRef<IMessageContext, ESPMode::ThreadSafe>& Context)
@@ -541,4 +476,28 @@ void FLiveLinkProvider::HandleHeartbeat(const FLiveLinkHeartbeatMessage& Message
 TSharedPtr<ILiveLinkProvider> ILiveLinkProvider::CreateLiveLinkProvider(const FString& ProviderName)
 {
 	return MakeShareable(new FLiveLinkProvider(ProviderName));
-}
\ No newline at end of file
+}
+
+void FLiveLinkProvider::CreateMessageEndpoint(FMessageEndpointBuilder& EndpointBuilder)
+{
+	MessageEndpoint = EndpointBuilder
+		.ReceivingOnAnyThread()
+		.Handling<FLiveLinkPingMessage>(this, &FLiveLinkProvider::HandlePingMessage)
+		.Handling<FLiveLinkConnectMessage>(this, &FLiveLinkProvider::HandleConnectMessage)
+		.Handling<FLiveLinkHeartbeatMessage>(this, &FLiveLinkProvider::HandleHeartbeat);
+
+	if (MessageEndpoint.IsValid())
+	{
+		MessageEndpoint->Subscribe<FLiveLinkPingMessage>();
+	}
+}
+
+void FLiveLinkProvider::GetConnectedAddresses(TArray<FMessageAddress>& Addresses)
+{
+	ValidateConnections();
+	Addresses.Reserve(ConnectedAddresses.Num());
+	for (const FTrackedAddress& Address : ConnectedAddresses)
+	{
+		Addresses.Add(Address.Address);
+	}
+}
diff --git a/Engine/Source/Runtime/LiveLinkMessageBusFramework/Public/LiveLinkProvider.h b/Engine/Source/Runtime/LiveLinkMessageBusFramework/Public/LiveLinkProvider.h
--- a/Engine/Source/Runtime/LiveLinkMessageBusFramework/Public/LiveLinkProvider.h
+++ b/Engine/Source/Runtime/LiveLinkMessageBusFramework/Public/LiveLinkProvider.h
@@ -40,6 +40,19 @@ struct ILiveLinkProvider : public ILiveLinkProvider_Base_DEPRECATED
 public:
 	LIVELINKMESSAGEBUSFRAMEWORK_API static TSharedPtr<ILiveLinkProvider> CreateLiveLinkProvider(const FString& ProviderName);
 
+	/**
+	 * Create a Live Link Provider based on a class derived from ILiveLinkProvider instead of using the default Live Link Provider.
+	 * @param ProviderName		The provider name.
+	 * @param EndpointBuilder	An endpoint builder that can be used to add additional message handlers.
+	 * @return					Shared pointer to the ILiveLinkProvider-derived class. Use StaticCastSharedPtr to cast it back to child class type.
+	 */
+	template<typename T>
+	static TSharedPtr<ILiveLinkProvider> CreateLiveLinkProvider(const FString& ProviderName,
+																struct FMessageEndpointBuilder&& EndpointBuilder)
+	{
+		return MakeShared<T>(ProviderName, MoveTemp(EndpointBuilder));
+	}
+
 	virtual ~ILiveLinkProvider() {}
 
 	/**
diff --git a/Engine/Source/Runtime/LiveLinkMessageBusFramework/Public/LiveLinkProviderImpl.h b/Engine/Source/Runtime/LiveLinkMessageBusFramework/Public/LiveLinkProviderImpl.h
new file mode 100644
--- /dev/null
+++ b/Engine/Source/Runtime/LiveLinkMessageBusFramework/Public/LiveLinkProviderImpl.h
@@ -0,0 +1,136 @@
+// Copyright Epic Games, Inc. All Rights Reserved.
+
+#pragma once
+
+#include "CoreMinimal.h"
+#include "IMessageContext.h"
+#include "LiveLinkMessages.h"
+#include "LiveLinkProvider.h"
+#include "MessageEndpoint.h"
+
+PRAGMA_DISABLE_DEPRECATION_WARNINGS
+// Subject that the application has told us about
+struct FTrackedSubject
+{
+	// Ref skeleton to go with transform data
+	FLiveLinkRefSkeleton RefSkeleton;
+
+	// Bone transform data
+	TArray<FTransform> Transforms;
+
+	// Curve data
+	TArray<FLiveLinkCurveElement> Curves;
+
+	// MetaData for subject
+	FLiveLinkMetaData MetaData;
+
+	// Incrementing time (application time) for interpolation purposes
+	double Time;
+};
+PRAGMA_ENABLE_DEPRECATION_WARNINGS
+
+struct FLiveLinkProvider : public ILiveLinkProvider
+{
+private:
+	const FString ProviderName;
+	const FString MachineName;
+
+	TSharedPtr<class FMessageEndpoint, ESPMode::ThreadSafe> MessageEndpoint;
+
+	/** Lock to stop multiple threads accessing the CurrentPreset at the same time */
+	mutable FCriticalSection CriticalSection;
+
+	// Array of our current connections
+	TArray<struct FTrackedAddress> ConnectedAddresses;
+
+	// Cache of our current subject state
+	TArray<struct FTrackedStaticData> StaticDatas;
+	TArray<struct FTrackedFrameData> FrameDatas;
+	TMap<FName, FTrackedSubject> Subjects;
+
+	// Delegate to notify interested parties when the client sources have changed
+	FLiveLinkProviderConnectionStatusChanged OnConnectionStatusChanged;
+
+	void CreateMessageEndpoint(struct FMessageEndpointBuilder& EndpointBuilder);
+
+	//Message bus message handlers
+	void HandlePingMessage(const FLiveLinkPingMessage& Message,
+						   const TSharedRef<class IMessageContext, ESPMode::ThreadSafe>& Context);
+	void HandleConnectMessage(const FLiveLinkConnectMessage& Message,
+							  const TSharedRef<class IMessageContext, ESPMode::ThreadSafe>& Context);
+	void HandleHeartbeat(const FLiveLinkHeartbeatMessage& Message,
+						 const TSharedRef<class IMessageContext, ESPMode::ThreadSafe>& Context);
+	// End message bus message handlers
+
+	// Validate our current connections
+	void ValidateConnections();
+	
+	FTrackedSubject& GetTrackedSubject(const FName& SubjectName);
+
+	void SendSubject(FName SubjectName, const struct FTrackedSubject& Subject);
+
+	void SendSubjectFrame(FName SubjectName, const struct FTrackedSubject& Subject);
+
+	// Get the cached data for the named subject
+	FTrackedStaticData* GetLastSubjectStaticData(const FName& SubjectName);
+
+	FTrackedFrameData* GetLastSubjectFrameData(const FName& SubjectName);
+
+	void SetLastSubjectStaticData(FName SubjectName, TSubclassOf<ULiveLinkRole> Role, FLiveLinkStaticDataStruct&& StaticData);
+
+	void SetLastSubjectFrameData(FName SubjectName, FLiveLinkFrameDataStruct&& FrameData);
+
+	// Clear a existing track subject
+	void ClearTrackedSubject(const FName& SubjectName);
+
+	void SendClearSubjectToConnections(FName SubjectName);
+
+	void GetConnectedAddresses(TArray<FMessageAddress>& Addresses);
+
+protected:
+	template<typename MessageType>
+	void SendMessage(MessageType* Message)
+	{
+		if (!Message)
+		{
+			return;
+		}
+
+		TArray<FMessageAddress> Addresses;
+		GetConnectedAddresses(Addresses);
+		if (Addresses.Num() != 0)
+		{
+			MessageEndpoint->Send(Message, Addresses);
+		}
+	}
+
+public:
+	FLiveLinkProvider(const FString& InProviderName);
+
+	FLiveLinkProvider(const FString& InProviderName, struct FMessageEndpointBuilder&& EndpointBuilder);
+
+	virtual ~FLiveLinkProvider();
+
+	virtual void UpdateSubject(const FName& SubjectName, const TArray<FName>& BoneNames, const TArray<int32>& BoneParents) override;
+
+	virtual bool UpdateSubjectStaticData(const FName SubjectName, TSubclassOf<ULiveLinkRole> Role, FLiveLinkStaticDataStruct&& StaticData) override;
+
+	virtual void ClearSubject(const FName& SubjectName);
+
+	virtual void RemoveSubject(const FName SubjectName) override;
+
+PRAGMA_DISABLE_DEPRECATION_WARNINGS
+	virtual void UpdateSubjectFrame(const FName& SubjectName, const TArray<FTransform>& BoneTransforms, const TArray<FLiveLinkCurveElement>& CurveData, double Time) override;
+
+	virtual void UpdateSubjectFrame(const FName& SubjectName, const TArray<FTransform>& BoneTransforms, const TArray<FLiveLinkCurveElement>& CurveData,
+									const FLiveLinkMetaData& MetaData, double Time) override;
+PRAGMA_ENABLE_DEPRECATION_WARNINGS
+
+	virtual bool UpdateSubjectFrameData(const FName SubjectName, FLiveLinkFrameDataStruct&& FrameData);
+
+	virtual bool HasConnection() const override;
+
+	virtual FDelegateHandle RegisterConnStatusChangedHandle(const FLiveLinkProviderConnectionStatusChanged::FDelegate& ConnStatusChanged) override;
+
+	virtual void UnregisterConnStatusChangedHandle(FDelegateHandle Handle);
+};
