PR #9446: Improvements for LiveLinkProvider derived classes (Contributed by [REDACTED_PERSON])

diff --git a/Engine/Source/Runtime/LiveLinkMessageBusFramework/Private/LiveLinkProvider.cpp b/Engine/Source/Runtime/LiveLinkMessageBusFramework/Private/LiveLinkProvider.cpp
--- a/Engine/Source/Runtime/LiveLinkMessageBusFramework/Private/LiveLinkProvider.cpp
+++ b/Engine/Source/Runtime/LiveLinkMessageBusFramework/Private/LiveLinkProvider.cpp
@@ -486,10 +486,7 @@ void FLiveLinkProvider::CreateMessageEndpoint(FMessageEndpointBuilder& EndpointB
 		.Handling<FLiveLinkConnectMessage>(this, &FLiveLinkProvider::HandleConnectMessage)
 		.Handling<FLiveLinkHeartbeatMessage>(this, &FLiveLinkProvider::HandleHeartbeat);
 
-	if (MessageEndpoint.IsValid())
-	{
-		MessageEndpoint->Subscribe<FLiveLinkPingMessage>();
-	}
+	Subscribe<FLiveLinkPingMessage>();
 }
 
 void FLiveLinkProvider::GetConnectedAddresses(TArray<FMessageAddress>& Addresses)
diff --git a/Engine/Source/Runtime/LiveLinkMessageBusFramework/Public/LiveLinkProviderImpl.h b/Engine/Source/Runtime/LiveLinkMessageBusFramework/Public/LiveLinkProviderImpl.h
--- a/Engine/Source/Runtime/LiveLinkMessageBusFramework/Public/LiveLinkProviderImpl.h
+++ b/Engine/Source/Runtime/LiveLinkMessageBusFramework/Public/LiveLinkProviderImpl.h
@@ -104,6 +104,36 @@ protected:
 		}
 	}
 
+	template<typename MessageType>
+	void SendMessage(MessageType* Message, const FMessageAddress& Address)
+	{
+		if (!Message || !Address.IsValid())
+		{
+			return;
+		}
+
+		MessageEndpoint->Send(Message, Address);
+	}
+
+	template<typename MessageType>
+	void Subscribe()
+	{
+		if (MessageEndpoint.IsValid())
+		{
+			MessageEndpoint->Subscribe<MessageType>();
+		}
+	}
+
+	const FString& GetProviderName() const
+	{
+		return ProviderName;
+	}
+	
+	const FString& GetMachineName() const
+	{
+		return MachineName;
+	}
+
 public:
 	FLiveLinkProvider(const FString& InProviderName);
 
