- Fixing compilation issue with LiveLinkProvider using LiveLinkMessageBusFramework

diff --git a/Engine/Source/Runtime/LiveLinkMessageBusFramework/LiveLinkMessageBusFramework.Build.cs b/Engine/Source/Runtime/LiveLinkMessageBusFramework/LiveLinkMessageBusFramework.Build.cs
--- a/Engine/Source/Runtime/LiveLinkMessageBusFramework/LiveLinkMessageBusFramework.Build.cs
+++ b/Engine/Source/Runtime/LiveLinkMessageBusFramework/LiveLinkMessageBusFramework.Build.cs
@@ -6,14 +6,14 @@ public class LiveLinkMessageBusFramework : ModuleRules
 {
     public LiveLinkMessageBusFramework(ReadOnlyTargetRules Target) : base(Target)
     {
-		PrivateDependencyModuleNames.Add("MessagingCommon");
 		PrivateIncludePathModuleNames.Add("Messaging");
 
 		PublicDependencyModuleNames.AddRange(
 			new string[] {
 				"Core",
 				"CoreUObject",
-				"LiveLinkInterface"
+				"LiveLinkInterface",
+				"MessagingCommon"
 			}
 		);
 	}
