Modify LiveLink to support extending LiveLinkProvider and LiveLinkMessageBusSource in LiveLinkHub
- Add ability to specify annotations in LiveLinkProvider's derived classes
- Allow creating LiveLinkProvider without creating endpoint in derived classes
- Move methods in LiveLinkMesageBusSource from private to protected to give more options to derived classes

diff --git a/Engine/Plugins/Animation/LiveLink/Source/LiveLink/Internal/ILiveLinkModule.h b/Engine/Plugins/Animation/LiveLink/Source/LiveLink/Internal/ILiveLinkModule.h
new file mode 100644
--- /dev/null
+++ b/Engine/Plugins/Animation/LiveLink/Source/LiveLink/Internal/ILiveLinkModule.h
@@ -0,0 +1,49 @@
+// Copyright Epic Games, Inc. All Rights Reserved.
+
+#pragma once
+
+#include "Modules/ModuleManager.h"
+
+#ifndef WITH_LIVELINK_DISCOVERY_MANAGER_THREAD
+#define WITH_LIVELINK_DISCOVERY_MANAGER_THREAD 1
+#endif
+
+class FLiveLinkHeartbeatEmitter;
+class FLiveLinkMessageBusDiscoveryManager;
+class FSlateStyleSet;
+
+/**
+ * Interface for messaging modules.
+ */
+class LIVELINK_API ILiveLinkModule
+	: public IModuleInterface
+{
+public:
+
+	/**
+	 * Gets a reference to the live link module instance.
+	 *
+	 * @return A reference to the live link module.
+	 */
+	static ILiveLinkModule& Get()
+	{
+#if PLATFORM_IOS
+		static ILiveLinkModule& LiveLinkModule = FModuleManager::LoadModuleChecked<ILiveLinkModule>("LiveLink");
+		return LiveLinkModule;
+#else
+		return FModuleManager::LoadModuleChecked<ILiveLinkModule>("LiveLink");
+#endif
+	}
+
+	virtual TSharedPtr<FSlateStyleSet> GetStyle() = 0;
+	virtual FLiveLinkHeartbeatEmitter& GetHeartbeatEmitter() = 0;
+#if WITH_LIVELINK_DISCOVERY_MANAGER_THREAD
+	virtual FLiveLinkMessageBusDiscoveryManager& GetMessageBusDiscoveryManager() = 0;
+#endif
+
+public:
+
+	/** Virtual destructor. */
+	virtual ~ILiveLinkModule() { }
+};
+
diff --git a/Engine/Plugins/Animation/LiveLink/Source/LiveLink/Internal/LiveLinkMessageBusDiscoveryManager.h b/Engine/Plugins/Animation/LiveLink/Source/LiveLink/Internal/LiveLinkMessageBusDiscoveryManager.h
new file mode 100644
--- /dev/null
+++ b/Engine/Plugins/Animation/LiveLink/Source/LiveLink/Internal/LiveLinkMessageBusDiscoveryManager.h
@@ -0,0 +1,64 @@
+﻿// Copyright Epic Games, Inc. All Rights Reserved.
+
+#pragma once
+
+#include "HAL/Runnable.h"
+#include "HAL/ThreadSafeBool.h"
+
+#include "LiveLinkMessageBusFinder.h"
+
+class IMessageContext;
+
+class FLiveLinkMessageBusSource;
+
+/** A class to asynchronously discover message bus sources. */
+class LIVELINK_API FLiveLinkMessageBusDiscoveryManager : FRunnable
+{
+public:
+	FLiveLinkMessageBusDiscoveryManager();
+	~FLiveLinkMessageBusDiscoveryManager();
+
+	//~ Begin FRunnable interface
+
+	virtual uint32 Run() override;
+
+	virtual void Stop() override;
+
+	//~ End FRunnable interface
+
+	void AddDiscoveryMessageRequest();
+	void RemoveDiscoveryMessageRequest();
+	TArray<FProviderPollResultPtr> GetDiscoveryResults() const;
+
+	bool IsRunning() const;
+
+private:
+	void HandlePongMessage(const FLiveLinkPongMessage& Message, const TSharedRef<IMessageContext, ESPMode::ThreadSafe>& Context);
+
+private:
+	// Counter of item that request discovery message
+	TAtomic<int32> PingRequestCounter;
+
+	// Last ping Request id
+	FGuid LastPingRequest;
+
+	// Time of the last ping request
+	double LastPingRequestTime;
+
+	// Ping request timeout
+	double PingRequestFrequency;
+
+	// Result from the last ping request
+	TArray<FProviderPollResultPtr> LastProviderPoolResults;
+
+	TSharedPtr<FMessageEndpoint, ESPMode::ThreadSafe> MessageEndpoint;
+
+	// Thread safe bool for stopping the thread
+	FThreadSafeBool bRunning;
+
+	// Thread the heartbeats are sent on
+	FRunnableThread* Thread;
+
+	// Critical section for accessing the Source Set
+	mutable FCriticalSection SourcesCriticalSection;
+};
diff --git a/Engine/Plugins/Animation/LiveLink/Source/LiveLink/Private/ILiveLinkModule.h b/Engine/Plugins/Animation/LiveLink/Source/LiveLink/Private/ILiveLinkModule.h
deleted file mode 100644
diff --git a/Engine/Plugins/Animation/LiveLink/Source/LiveLink/Private/LiveLinkMessageBusDiscoveryManager.cpp b/Engine/Plugins/Animation/LiveLink/Source/LiveLink/Private/LiveLinkMessageBusDiscoveryManager.cpp
--- a/Engine/Plugins/Animation/LiveLink/Source/LiveLink/Private/LiveLinkMessageBusDiscoveryManager.cpp
+++ b/Engine/Plugins/Animation/LiveLink/Source/LiveLink/Private/LiveLinkMessageBusDiscoveryManager.cpp
@@ -114,6 +114,7 @@ void FLiveLinkMessageBusDiscoveryManager::HandlePongMessage(const FLiveLinkPongM
 		// Verify Message.LiveLinkVersion to consider validity of discovered provider. Older UE always sends 1
 		constexpr bool bIsValidProvider = true;
 		const double MachineTimeOffset = LiveLinkMessageBusHelper::CalculateProviderMachineOffset(Message.CreationPlatformTime, Context);
-		LastProviderPoolResults.Emplace(MakeShared<FProviderPollResult, ESPMode::ThreadSafe>(Context->GetSender(), Message.ProviderName, Message.MachineName, MachineTimeOffset, bIsValidProvider));
+
+		LastProviderPoolResults.Emplace(MakeShared<FProviderPollResult, ESPMode::ThreadSafe>(Context->GetSender(), Message.ProviderName, Message.MachineName, MachineTimeOffset, bIsValidProvider, Context->GetAnnotations()));
 	}
 }
diff --git a/Engine/Plugins/Animation/LiveLink/Source/LiveLink/Private/LiveLinkMessageBusDiscoveryManager.h b/Engine/Plugins/Animation/LiveLink/Source/LiveLink/Private/LiveLinkMessageBusDiscoveryManager.h
deleted file mode 100644
diff --git a/Engine/Plugins/Animation/LiveLink/Source/LiveLink/Private/LiveLinkMessageBusSource.cpp b/Engine/Plugins/Animation/LiveLink/Source/LiveLink/Private/LiveLinkMessageBusSource.cpp
--- a/Engine/Plugins/Animation/LiveLink/Source/LiveLink/Private/LiveLinkMessageBusSource.cpp
+++ b/Engine/Plugins/Animation/LiveLink/Source/LiveLink/Private/LiveLinkMessageBusSource.cpp
@@ -18,12 +18,13 @@
 #include "MessageEndpointBuilder.h"
 #include "Misc/App.h"
 
+
 FLiveLinkMessageBusSource::FLiveLinkMessageBusSource(const FText& InSourceType, const FText& InSourceMachineName, const FMessageAddress& InConnectionAddress, double InMachineTimeOffset)
 	: ConnectionAddress(InConnectionAddress)
+	, bIsValid(false)
 	, SourceType(InSourceType)
 	, SourceMachineName(InSourceMachineName)
 	, ConnectionLastActive(0.0)
-	, bIsValid(false)
 	, MachineTimeOffset(InMachineTimeOffset)
 {}
 
@@ -88,7 +89,7 @@ void FLiveLinkMessageBusSource::Update()
 		bIsValid = CurrentTime - ConnectionLastActive < HeartbeatTimeout;
 		if (!bIsValid)
 		{
-			const double DeadSourceTimeout = GetDefault<ULiveLinkSettings>()->GetMessageBusTimeBeforeRemovingDeadSource();
+			const double DeadSourceTimeout = GetDeadSourceTimeout();
 			if (CurrentTime - ConnectionLastActive > DeadSourceTimeout)
 			{
 				if (RequestSourceShutdown())
@@ -219,6 +220,12 @@ TSubclassOf<ULiveLinkSourceSettings> FLiveLinkMessageBusSource::GetSettingsClass
 	return ULiveLinkMessageBusSourceSettings::StaticClass();
 }
 
+void FLiveLinkMessageBusSource::StartHeartbeatEmitter()
+{
+	FLiveLinkHeartbeatEmitter& HeartbeatEmitter = ILiveLinkModule::Get().GetHeartbeatEmitter();
+	HeartbeatEmitter.StartHeartbeat(ConnectionAddress, MessageEndpoint);
+}
+
 TSharedPtr<FMessageEndpoint, ESPMode::ThreadSafe> FLiveLinkMessageBusSource::CreateAndInitializeMessageEndpoint()
 {
 	FMessageEndpointBuilder EndpointBuilder = FMessageEndpoint::Builder(GetSourceName());
@@ -253,6 +260,11 @@ void FLiveLinkMessageBusSource::InitializeAndPushFrameData_AnyThread(FName Subje
 	PushClientSubjectFrameData_AnyThread(SubjectKey, MoveTemp(DataStruct));
 }
 
+double FLiveLinkMessageBusSource::GetDeadSourceTimeout() const
+{
+	return GetDefault<ULiveLinkSettings>()->GetMessageBusTimeBeforeRemovingDeadSource();
+}
+
 void FLiveLinkMessageBusSource::PushClientSubjectStaticData_AnyThread(const FLiveLinkSubjectKey& SubjectKey,
 																	  TSubclassOf<ULiveLinkRole> Role,
 																	  FLiveLinkStaticDataStruct&& StaticData)
@@ -294,8 +306,7 @@ void FLiveLinkMessageBusSource::SendConnectMessage()
 	FLiveLinkConnectMessage* ConnectMessage = FMessageEndpoint::MakeMessage<FLiveLinkConnectMessage>();
 	ConnectMessage->LiveLinkVersion = ILiveLinkClient::LIVELINK_VERSION;
 	SendMessage(ConnectMessage);
-	FLiveLinkHeartbeatEmitter& HeartbeatEmitter = ILiveLinkModule::Get().GetHeartbeatEmitter();
-	HeartbeatEmitter.StartHeartbeat(ConnectionAddress, MessageEndpoint);
+	StartHeartbeatEmitter();
 	bIsValid = true;
 }
 
diff --git a/Engine/Plugins/Animation/LiveLink/Source/LiveLink/Public/LiveLinkMessageBusFinder.h b/Engine/Plugins/Animation/LiveLink/Source/LiveLink/Public/LiveLinkMessageBusFinder.h
--- a/Engine/Plugins/Animation/LiveLink/Source/LiveLink/Public/LiveLinkMessageBusFinder.h
+++ b/Engine/Plugins/Animation/LiveLink/Source/LiveLink/Public/LiveLinkMessageBusFinder.h
@@ -40,8 +40,19 @@ public:
 		, MachineName(InMachineName)
 		, MachineTimeOffset(InMachineOffset)
 		, bIsValidProvider(bInIsValidProvider)
+		, Annotations()
 	{}
 
+	FProviderPollResult(const FMessageAddress& InAddress, const FString& InName, const FString& InMachineName, double InMachineOffset, bool bInIsValidProvider, TMap<FName, FString> InAnnotations)
+		: Address(InAddress)
+		, Name(InName)
+		, MachineName(InMachineName)
+		, MachineTimeOffset(InMachineOffset)
+		, bIsValidProvider(bInIsValidProvider)
+		, Annotations(MoveTemp(InAnnotations))
+	{}
+
+	// Messagebus address of the provider
 	FMessageAddress Address;
 
 	// The name of the provider
@@ -59,6 +70,9 @@ public:
 	// Whether the provider is valid (compatible with the current version of LiveLink)
 	UPROPERTY(BlueprintReadOnly, VisibleAnywhere, Category = "LiveLink")
 	bool bIsValidProvider = true;
+
+	// Annotations present on the message
+	TMap<FName, FString> Annotations;
 };
 
 typedef TSharedPtr<FProviderPollResult, ESPMode::ThreadSafe> FProviderPollResultPtr;
diff --git a/Engine/Plugins/Animation/LiveLink/Source/LiveLink/Public/LiveLinkMessageBusSource.h b/Engine/Plugins/Animation/LiveLink/Source/LiveLink/Public/LiveLinkMessageBusSource.h
--- a/Engine/Plugins/Animation/LiveLink/Source/LiveLink/Public/LiveLinkMessageBusSource.h
+++ b/Engine/Plugins/Animation/LiveLink/Source/LiveLink/Public/LiveLinkMessageBusSource.h
@@ -59,11 +59,17 @@ protected:
 													  const TSharedRef<IMessageContext, ESPMode::ThreadSafe>& Context,
 													  UScriptStruct* MessageTypeInfo);
 
+	// Allows derived classes to provide their own timeout duration before a source is removed because the heartbeat timeout was hit
+	virtual double GetDeadSourceTimeout() const;
+
 	// Send the static data to the clients
 	void PushClientSubjectStaticData_AnyThread(const FLiveLinkSubjectKey& SubjectKey, TSubclassOf<ULiveLinkRole> Role, FLiveLinkStaticDataStruct&& StaticData);
 	// Send the frame data to the clients
 	void PushClientSubjectFrameData_AnyThread(const FLiveLinkSubjectKey& SubjectKey, FLiveLinkFrameDataStruct&& FrameData);
 
+	// Send connect message to the provider and start the heartbeat emitter
+	virtual void SendConnectMessage();
+
 	// Send a message through the endpoint
 	template<typename MessageType>
 	void SendMessage(MessageType* Message)
@@ -76,6 +82,18 @@ protected:
 		MessageEndpoint->Send(Message, ConnectionAddress);
 	}
 
+	// Start the heartbeat emitter for this connection
+	void StartHeartbeatEmitter();
+protected:
+	// Message bus endpoint responsible for communication with the livelink provider
+	TSharedPtr<FMessageEndpoint, ESPMode::ThreadSafe> MessageEndpoint;
+
+	// Connection address of the livelink provider 
+	FMessageAddress ConnectionAddress;
+
+	// Current Validity of Source
+	FThreadSafeBool bIsValid;
+
 private:
 	TSharedPtr<FMessageEndpoint, ESPMode::ThreadSafe> CreateAndInitializeMessageEndpoint();
 
@@ -93,8 +111,6 @@ private:
 	// Threadsafe update of the last active time
 	FORCEINLINE void UpdateConnectionLastActive();
 
-	void SendConnectMessage();
-
 	ILiveLinkClient* Client;
 
 	// Our identifier in LiveLink
@@ -103,19 +119,12 @@ private:
 	// List of the roles available when the bus was opened
 	TArray<TWeakObjectPtr<ULiveLinkRole>> RoleInstances;
 
-	TSharedPtr<FMessageEndpoint, ESPMode::ThreadSafe> MessageEndpoint;
-
-	FMessageAddress ConnectionAddress;
-
 	FText SourceType;
 	FText SourceMachineName;
 
 	// Time we last received anything 
 	double ConnectionLastActive;
 
-	// Current Validity of Source
-	FThreadSafeBool bIsValid;
-
 	// Critical section to allow for threadsafe updating of the connection time
 	FCriticalSection ConnectionLastActiveSection;
 
diff --git a/Engine/Plugins/Animation/LiveLink/Source/LiveLinkEditor/Private/LiveLinkPanelController.cpp b/Engine/Plugins/Animation/LiveLink/Source/LiveLinkEditor/Private/LiveLinkPanelController.cpp
--- a/Engine/Plugins/Animation/LiveLink/Source/LiveLinkEditor/Private/LiveLinkPanelController.cpp
+++ b/Engine/Plugins/Animation/LiveLink/Source/LiveLinkEditor/Private/LiveLinkPanelController.cpp
@@ -115,7 +115,9 @@ void FLiveLinkPanelController::OnSubjectSelectionChanged(FLiveLinkSubjectUIEntry
 
 	if (!bDetailViewSet)
 	{
+#if !WITH_LIVELINK_HUB
 		SourcesDetailsView->SetObject(nullptr);
+#endif
 		SubjectsDetailsView->SetSubjectKey(FLiveLinkSubjectKey());
 	}
 
diff --git a/Engine/Source/Runtime/LiveLinkMessageBusFramework/Private/LiveLinkProvider.cpp b/Engine/Source/Runtime/LiveLinkMessageBusFramework/Private/LiveLinkProvider.cpp
--- a/Engine/Source/Runtime/LiveLinkMessageBusFramework/Private/LiveLinkProvider.cpp
+++ b/Engine/Source/Runtime/LiveLinkMessageBusFramework/Private/LiveLinkProvider.cpp
@@ -3,6 +3,7 @@
 #include "LiveLinkProvider.h"
 #include "LiveLinkProviderImpl.h"
 
+#include "Algo/RemoveIf.h"
 #include "HAL/PlatformProcess.h"
 #include "IMessageContext.h"
 #include "LiveLinkMessages.h"
@@ -93,10 +94,20 @@ void FLiveLinkProvider::ValidateConnections()
 {
 	FConnectionValidator Validator;
 
-	const int32 RemovedConnections = ConnectedAddresses.RemoveAll([=](const FTrackedAddress& Address) { return !Validator(Address); });
+	TArray<FMessageAddress> RemovedConnections;
+	Algo::RemoveIf(ConnectedAddresses, [this, &Validator, &RemovedConnections](const FTrackedAddress& Address) mutable
+	{
+		if (!Validator(Address))
+	    {
+			RemovedConnections.Add(Address.Address);
+			return true;
+	    }
+		return false;
+	});
 
-	if (RemovedConnections > 0)
+	if (RemovedConnections.Num() > 0)
 	{
+		OnConnectionsClosed(RemovedConnections);
 		OnConnectionStatusChanged.Broadcast();
 	}
 }
@@ -117,7 +128,7 @@ void FLiveLinkProvider::SendSubject(FName SubjectName, const FTrackedSubject& Su
 	TArray<FMessageAddress> Addresses;
 	GetConnectedAddresses(Addresses);
 
-	MessageEndpoint->Send(SubjectData, Addresses);
+	MessageEndpoint->Send(SubjectData, FLiveLinkSubjectDataMessage::StaticStruct(), EMessageFlags::None, GetAnnotations(), nullptr, Addresses, FTimespan::Zero(), FDateTime::MaxValue());
 }
 
 // Send frame data for named subject
@@ -133,7 +144,7 @@ void FLiveLinkProvider::SendSubjectFrame(FName SubjectName, const FTrackedSubjec
 	TArray<FMessageAddress> Addresses;
 	GetConnectedAddresses(Addresses);
 
-	MessageEndpoint->Send(SubjectFrame, Addresses);
+	MessageEndpoint->Send(SubjectFrame, FLiveLinkSubjectFrameMessage::StaticStruct(), EMessageFlags::None, GetAnnotations(), nullptr, Addresses, FTimespan::Zero(), FDateTime::MaxValue());
 }
 
 // Get the cached data for the named subject
@@ -195,7 +206,7 @@ void FLiveLinkProvider::SendClearSubjectToConnections(FName SubjectName)
 	TArray<FMessageAddress> MessageAddresses;
 	GetConnectedAddresses(MessageAddresses);
 
-	MessageEndpoint->Send(FMessageEndpoint::MakeMessage<FLiveLinkClearSubject>(SubjectName), EMessageFlags::Reliable, nullptr, MessageAddresses, FTimespan::Zero(), FDateTime::MaxValue());
+	MessageEndpoint->Send(FMessageEndpoint::MakeMessage<FLiveLinkClearSubject>(SubjectName), EMessageFlags::Reliable, GetAnnotations(), nullptr, MessageAddresses, FTimespan::Zero(), FDateTime::MaxValue());
 }
 
 FLiveLinkProvider::FLiveLinkProvider(const FString& InProviderName)
@@ -213,6 +224,17 @@ FLiveLinkProvider::FLiveLinkProvider(const FString& InProviderName, FMessageEndp
 	CreateMessageEndpoint(EndpointBuilder);
 }
 
+FLiveLinkProvider::FLiveLinkProvider(const FString& InProviderName, bool bInCreateEndpoint)
+	: ProviderName(InProviderName)
+	, MachineName(FPlatformProcess::ComputerName())
+{
+	if (bInCreateEndpoint)
+	{
+		FMessageEndpointBuilder EndpointBuilder = FMessageEndpoint::Builder(*InProviderName);
+    	CreateMessageEndpoint(EndpointBuilder);
+	}
+}
+
 FLiveLinkProvider::~FLiveLinkProvider()
 {
 	if (MessageEndpoint.IsValid())
@@ -404,7 +426,7 @@ void FLiveLinkProvider::HandlePingMessage(const FLiveLinkPingMessage& Message, c
 		return;
 	}
 
-	MessageEndpoint->Send(FMessageEndpoint::MakeMessage<FLiveLinkPongMessage>(ProviderName, MachineName, Message.PollRequest, LIVELINK_SupportedVersion), Context->GetSender());
+	MessageEndpoint->Send(FMessageEndpoint::MakeMessage<FLiveLinkPongMessage>(ProviderName, MachineName, Message.PollRequest, LIVELINK_SupportedVersion), GetAnnotations(), Context->GetSender());
 }
 
 void FLiveLinkProvider::HandleConnectMessage(const FLiveLinkConnectMessage& Message, const TSharedRef<IMessageContext, ESPMode::ThreadSafe>& Context)
@@ -435,7 +457,7 @@ void FLiveLinkProvider::HandleConnectMessage(const FLiveLinkConnectMessage& Mess
 		TArray<FMessageAddress> MessageAddress;
 		MessageAddress.Add(ConnectionAddress);
 
-		TMap<FName, FString> Annotations;
+		TMap<FName, FString> Annotations = GetAnnotations();
 		Annotations.Add(FLiveLinkMessageAnnotation::SubjectAnnotation, TEXT(""));
 		Annotations.Add(FLiveLinkMessageAnnotation::RoleAnnotation, TEXT(""));
 
@@ -469,7 +491,7 @@ void FLiveLinkProvider::HandleHeartbeat(const FLiveLinkHeartbeatMessage& Message
 		TrackedAddress->LastHeartbeatTime = FPlatformTime::Seconds();
 
 		// Respond so editor gets heartbeat too
-		MessageEndpoint->Send(FMessageEndpoint::MakeMessage<FLiveLinkHeartbeatMessage>(), Context->GetSender());
+		MessageEndpoint->Send(FMessageEndpoint::MakeMessage<FLiveLinkHeartbeatMessage>(), GetAnnotations(), Context->GetSender());
 	}
 }
 
diff --git a/Engine/Source/Runtime/LiveLinkMessageBusFramework/Public/LiveLinkProviderImpl.h b/Engine/Source/Runtime/LiveLinkMessageBusFramework/Public/LiveLinkProviderImpl.h
--- a/Engine/Source/Runtime/LiveLinkMessageBusFramework/Public/LiveLinkProviderImpl.h
+++ b/Engine/Source/Runtime/LiveLinkMessageBusFramework/Public/LiveLinkProviderImpl.h
@@ -51,13 +51,10 @@ private:
 	// Delegate to notify interested parties when the client sources have changed
 	FLiveLinkProviderConnectionStatusChanged OnConnectionStatusChanged;
 
-	void CreateMessageEndpoint(struct FMessageEndpointBuilder& EndpointBuilder);
-
+private:
 	//Message bus message handlers
 	void HandlePingMessage(const FLiveLinkPingMessage& Message,
 						   const TSharedRef<class IMessageContext, ESPMode::ThreadSafe>& Context);
-	void HandleConnectMessage(const FLiveLinkConnectMessage& Message,
-							  const TSharedRef<class IMessageContext, ESPMode::ThreadSafe>& Context);
 	void HandleHeartbeat(const FLiveLinkHeartbeatMessage& Message,
 						 const TSharedRef<class IMessageContext, ESPMode::ThreadSafe>& Context);
 	// End message bus message handlers
@@ -83,6 +80,12 @@ private:
 	void SendClearSubjectToConnections(FName SubjectName);
 
 protected:
+	// Update connected addresses and send information to the connected source
+	void HandleConnectMessage(const FLiveLinkConnectMessage& Message, const TSharedRef<class IMessageContext, ESPMode::ThreadSafe>& Context);
+
+	// Create the message bus message endoing responsble for dispatching message bus messages to their respective handlers
+	void CreateMessageEndpoint(struct FMessageEndpointBuilder& EndpointBuilder);
+	
 	// Get the addresses of all connected instances.
 	void GetConnectedAddresses(TArray<FMessageAddress>& Addresses);
 
@@ -135,6 +138,18 @@ protected:
 		return MachineName;
 	}
 
+	// Called after ValidateConnections removes invalid connections 
+	virtual void OnConnectionsClosed(const TArray<FMessageAddress>& ClosedAddresses) {}
+
+	// Get annotations to include on every message sent by this provider
+	virtual TMap<FName, FString> GetAnnotations() const
+	{
+		return {};
+	}
+
+	// Constructor for derived classes that allows specifying that no endpoint should be created.
+	FLiveLinkProvider(const FString& InProviderName, bool bInCreateEndpoint);
+
 public:
 	FLiveLinkProvider(const FString& InProviderName);
 
