LiveLink - Add option to override a source name in LiveLink to disambiguate LiveLinkHub's sources

diff --git a/Engine/Plugins/Animation/LiveLink/Source/LiveLink/Private/LiveLinkClient.cpp b/Engine/Plugins/Animation/LiveLink/Source/LiveLink/Private/LiveLinkClient.cpp
--- a/Engine/Plugins/Animation/LiveLink/Source/LiveLink/Private/LiveLinkClient.cpp
+++ b/Engine/Plugins/Animation/LiveLink/Source/LiveLink/Private/LiveLinkClient.cpp
@@ -8,6 +8,7 @@
 #include "Engine/Engine.h"
 #include "LiveLinkAnimationVirtualSubject.h"
 #include "LiveLinkLog.h"
+#include "LiveLinkMessages.h"
 #include "LiveLinkPresetTypes.h"
 #include "LiveLinkRoleTrait.h"
 #include "LiveLinkProvider.h"
@@ -318,14 +319,19 @@ void FLiveLinkClient::HandleSubjectRebroadcast(ILiveLinkSubject* InSubject, cons
 
 				const FName RebroadcastName = GetRebroadcastName(InSubject->GetSubjectKey());
 
+				const FText OriginalSourceType = GetSourceType(InSubject->GetSubjectKey().Source);
+
+				TMap<FName, FString> ExtraAnnotations;
+				ExtraAnnotations.Add(FLiveLinkMessageAnnotation::OriginalSourceAnnotation, OriginalSourceType.ToString());
+
 				if (!InSubject->HasStaticDataBeenRebroadcasted())
 				{
-					RebroadcastLiveLinkProvider->UpdateSubjectStaticData(RebroadcastName, SubjectRole, MoveTemp(StaticDataCopy));
+					RebroadcastLiveLinkProvider->UpdateSubjectStaticData(RebroadcastName, SubjectRole, MoveTemp(StaticDataCopy), ExtraAnnotations);
 					InSubject->SetStaticDataAsRebroadcasted(true);
 					RebroadcastedSubjects.Add(InSubject->GetSubjectKey());
 				}
 
-				RebroadcastLiveLinkProvider->UpdateSubjectFrameData(RebroadcastName, MoveTemp(FrameDataCopy));
+				RebroadcastLiveLinkProvider->UpdateSubjectFrameData(RebroadcastName, MoveTemp(FrameDataCopy), ExtraAnnotations);
 			}
 			else
 			{
@@ -665,32 +671,7 @@ FLiveLinkSourcePreset FLiveLinkClient::GetSourcePreset(FGuid InSourceGuid, UObje
 void FLiveLinkClient::PushSubjectStaticData_AnyThread(const FLiveLinkSubjectKey& InSubjectKey, TSubclassOf<ULiveLinkRole> InRole, FLiveLinkStaticDataStruct&& InStaticData)
 {
 	FPendingSubjectStatic SubjectStatic{ InSubjectKey, InRole, MoveTemp(InStaticData) };
-	const int32 MaxNumBufferToCached = CVarMaxNewStaticDataPerUpdate.GetValueOnAnyThread();
-	bool bLogError = false;
-	{
-		FScopeLock Lock(&PendingFramesCriticalSection);
-		if (SubjectStaticToPush.Num() > MaxNumBufferToCached) // Something is wrong somewhere. Warn the user and discard the new Static Data.
-		{
-			bLogError = true;
-		}
-		else
-		{
-			{
-				FScopeLock BroadcastLock(&SubjectFrameReceivedHandleseCriticalSection);
-				if (const FSubjectFramesReceivedHandles* Handles = SubjectFrameReceivedHandles.Find(InSubjectKey))
-				{
-					Handles->OnStaticDataReceived.Broadcast(SubjectStatic.StaticData);
-				}
-			}
-			SubjectStaticToPush.Add(MoveTemp(SubjectStatic));
-		}
-	}
-
-	if (bLogError)
-	{
-		static const FName NAME_TooManyStatic = "LiveLinkClient_TooManyStatic";
-		FLiveLinkLog::ErrorOnce(NAME_TooManyStatic, FLiveLinkSubjectKey(), TEXT("Trying to add more than %d static subjects in the same frame. New Subjects will be discarded."), MaxNumBufferToCached);
-	}
+	PushPendingSubject_AnyThread(MoveTemp(SubjectStatic));
 }
 
 void FLiveLinkClient::PushSubjectStaticData_Internal(FPendingSubjectStatic&& SubjectStaticData)
@@ -782,6 +763,11 @@ void FLiveLinkClient::PushSubjectStaticData_Internal(FPendingSubjectStatic&& Sub
 
 			SubjectSettings->Role = SubjectStaticData.Role;
 
+			if (FString* OriginalSourceName = SubjectStaticData.ExtraMetadata.Find(FLiveLinkMessageAnnotation::OriginalSourceAnnotation))
+			{
+				SubjectSettings->OriginalSourceName = **OriginalSourceName;
+			}
+
 			UClass* FrameInterpolationProcessorClass = DefaultSetting.FrameInterpolationProcessor.Get();
 			if (FrameInterpolationProcessorClass != nullptr)
 			{
@@ -1421,6 +1407,42 @@ TArray<FLiveLinkTime> FLiveLinkClient::GetSubjectFrameTimes(FLiveLinkSubjectName
 	return TArray<FLiveLinkTime>();
 }
 
+FText FLiveLinkClient::GetSourceNameOverride(const FLiveLinkSubjectKey& SubjectKey) const
+{
+	FText SourceType = GetSourceType(SubjectKey.Source);
+	FText SourceNameOverride = SourceType;
+
+	UObject* Settings = GetSubjectSettings(SubjectKey);
+	if (ULiveLinkSubjectSettings* SubjectSettings = Cast<ULiveLinkSubjectSettings>(Settings))
+	{
+		if (!SubjectSettings->OriginalSourceName.IsNone())
+		{
+			SourceNameOverride = FText::Format(INVTEXT("{0} ({1})"), FText::FromName(SubjectSettings->OriginalSourceName), SourceType);
+		}
+	}
+
+	return SourceNameOverride;
+}
+
+FText FLiveLinkClient::GetSubjectDisplayName(const FLiveLinkSubjectKey& InSubjectKey) const
+{
+	FText DisplayName;
+	if (const FLiveLinkCollectionSubjectItem* SubjectItem = Collection->FindSubject(InSubjectKey))
+	{
+		UObject* Settings = SubjectItem->GetSettings();
+		if (ULiveLinkSubjectSettings* SubjectSettings = Cast<ULiveLinkSubjectSettings>(Settings))
+		{
+			DisplayName = SubjectSettings->GetDisplayName();
+		}
+		else if (ULiveLinkVirtualSubject* VirtualSubject = Cast<ULiveLinkVirtualSubject>(Settings))
+		{
+			VirtualSubject->GetDisplayName();
+		}
+	}
+
+	return DisplayName;
+}
+
 TArray<FLiveLinkSubjectKey> FLiveLinkClient::GetSubjectsSupportingRole(TSubclassOf<ULiveLinkRole> InSupportedRole, bool bIncludeDisabledSubject, bool bIncludeVirtualSubject) const
 {
 	TArray<FLiveLinkSubjectKey> SubjectKeys;
@@ -1636,6 +1658,35 @@ FName FLiveLinkClient::GetRebroadcastName(const FLiveLinkSubjectKey& InSubjectKe
 	return InSubjectKey.SubjectName;
 }
 
+void FLiveLinkClient::PushPendingSubject_AnyThread(FPendingSubjectStatic&& PendingSubject)
+{
+	const int32 MaxNumBufferToCached = CVarMaxNewStaticDataPerUpdate.GetValueOnAnyThread();
+	bool bLogError = true;
+	{
+		FScopeLock Lock(&PendingFramesCriticalSection);
+		if (SubjectStaticToPush.Num() <= MaxNumBufferToCached) 
+		{
+			bLogError = false;
+
+			{
+				FScopeLock BroadcastLock(&SubjectFrameReceivedHandleseCriticalSection);
+				if (const FSubjectFramesReceivedHandles* Handles = SubjectFrameReceivedHandles.Find(PendingSubject.SubjectKey))
+				{
+					Handles->OnStaticDataReceived.Broadcast(PendingSubject.StaticData);
+				}
+			}
+			SubjectStaticToPush.Add(MoveTemp(PendingSubject));
+		}
+	}
+
+	if (bLogError)
+	{
+		// Something is wrong somewhere. Warn the user and discard the new Static Data.
+		static const FName NAME_TooManyStatic = "LiveLinkClient_TooManyStatic";
+		FLiveLinkLog::ErrorOnce(NAME_TooManyStatic, FLiveLinkSubjectKey(), TEXT("Trying to add more than %d static subjects in the same frame. New Subjects will be discarded."), MaxNumBufferToCached);
+	}
+}
+
 FText FLiveLinkClient::GetSourceType(FGuid InEntryGuid) const
 {
 	if (const FLiveLinkCollectionSourceItem* SourceItem = Collection->FindSource(InEntryGuid))
@@ -1706,7 +1757,7 @@ UObject* FLiveLinkClient::GetSubjectSettings(const FLiveLinkSubjectKey& InSubjec
 		return SubjectItem->GetSettings();
 	}
 	return nullptr;
-}
+}	
 
 const FLiveLinkStaticDataStruct* FLiveLinkClient::GetSubjectStaticData_AnyThread(const FLiveLinkSubjectKey& InSubjectKey) const
 {
diff --git a/Engine/Plugins/Animation/LiveLink/Source/LiveLink/Private/LiveLinkSourceCollection.cpp b/Engine/Plugins/Animation/LiveLink/Source/LiveLink/Private/LiveLinkSourceCollection.cpp
--- a/Engine/Plugins/Animation/LiveLink/Source/LiveLink/Private/LiveLinkSourceCollection.cpp
+++ b/Engine/Plugins/Animation/LiveLink/Source/LiveLink/Private/LiveLinkSourceCollection.cpp
@@ -267,6 +267,13 @@ const FLiveLinkCollectionSubjectItem* FLiveLinkSourceCollection::FindSubject(FLi
 	return const_cast<FLiveLinkSourceCollection*>(this)->FindSubject(InSubjectKey);
 }
 
+
+const FLiveLinkCollectionSubjectItem* FLiveLinkSourceCollection::FindSubject(FLiveLinkSubjectName SubjectName) const
+{
+	UE::TUniqueLock Lock(SubjectsLock);
+	return Subjects.FindByPredicate([SubjectName](const FLiveLinkCollectionSubjectItem& Other) { return Other.Key.SubjectName == SubjectName;  });
+}
+
 const FLiveLinkCollectionSubjectItem* FLiveLinkSourceCollection::FindEnabledSubject(FLiveLinkSubjectName InSubjectName) const
 {
 	UE::TUniqueLock Lock(SubjectsLock);
diff --git a/Engine/Plugins/Animation/LiveLink/Source/LiveLink/Public/LiveLinkClient.h b/Engine/Plugins/Animation/LiveLink/Source/LiveLink/Public/LiveLinkClient.h
--- a/Engine/Plugins/Animation/LiveLink/Source/LiveLink/Public/LiveLinkClient.h
+++ b/Engine/Plugins/Animation/LiveLink/Source/LiveLink/Public/LiveLinkClient.h
@@ -120,6 +120,8 @@ public:
 	virtual TArray<FLiveLinkSubjectKey> GetSubjectsSupportingRole(TSubclassOf<ULiveLinkRole> SupportedRole, bool bIncludeDisabledSubject, bool bIncludeVirtualSubject) const override;
 	virtual TArray<FLiveLinkTime> GetSubjectFrameTimes(const FLiveLinkSubjectKey& SubjectKey) const override;
 	virtual TArray<FLiveLinkTime> GetSubjectFrameTimes(FLiveLinkSubjectName SubjectName) const override;
+	virtual FText GetSourceNameOverride(const FLiveLinkSubjectKey& SubjectKey) const override;
+	virtual FText GetSubjectDisplayName(const FLiveLinkSubjectKey& SubjectKey) const override;
 	virtual ULiveLinkSourceSettings* GetSourceSettings(const FGuid& SourceGuid) const override;
 	virtual UObject* GetSubjectSettings(const FLiveLinkSubjectKey& SubjectKey) const override;
 	virtual const FLiveLinkStaticDataStruct* GetSubjectStaticData_AnyThread(const FLiveLinkSubjectKey& InSubjectKey) const override;
@@ -150,6 +152,24 @@ public:
 	virtual void UnregisterSubjectFramesHandle(FLiveLinkSubjectName InSubjectName, FDelegateHandle InStaticDataReceivedHandle, FDelegateHandle InFrameDataReceivedHandle) override;
 	//~ End ILiveLinkClient implementation
 
+public:
+	/** Struct that hold the pending static data that will be pushed next tick. */
+	struct FPendingSubjectStatic
+	{
+		FLiveLinkSubjectKey SubjectKey;
+		TSubclassOf<ULiveLinkRole> Role;
+		FLiveLinkStaticDataStruct StaticData;
+		TMap<FName, FString> ExtraMetadata;
+	};
+
+	/** Struct that hold the pending frame data that will be pushed next tick. */
+	struct FPendingSubjectFrame
+	{
+		FLiveLinkSubjectKey SubjectKey;
+		FLiveLinkFrameDataStruct FrameData;
+	};
+
+
 	/** The tick callback to update the pending work and clear the subject's snapshot*/
 	void Tick();
 
@@ -180,9 +200,12 @@ public:
 
 	FLiveLinkSubjectTimeSyncData GetTimeSyncData(FLiveLinkSubjectName SubjectName);
 
-	/** Get the rebroadcast name for a given subject. (Defauls to the subject's subject name, but can be overriden. */
+	/** Get the rebroadcast name for a given subject. (Defaults to the subject's subject name, but can be overriden. */
 	FName GetRebroadcastName(const FLiveLinkSubjectKey& InSubjectKey) const;
 
+	/** Push subject static data with additional metadata. */
+	void PushPendingSubject_AnyThread(FPendingSubjectStatic&& PendingSubject);
+
 	UE_DEPRECATED(4.23, "FLiveLinkClient::GetSourceTypeForEntry is deprecated. Please use GetSourceType instead!")
 	FText GetSourceTypeForEntry(FGuid EntryGuid) const { return GetSourceType(EntryGuid); }
 	UE_DEPRECATED(4.23, "FLiveLinkClient::GetMachineNameForEntry is deprecated. Please use GetSourceMachineName instead!")
@@ -214,22 +237,6 @@ protected:
 	/** Remove the delegates that were triggered for all subjects. */
 	void UnregisterGlobalSubjectFramesDelegate(FDelegateHandle& InStaticDataAddedHandle, FDelegateHandle& InFrameDataAddedHandle);
 
-private:
-	/** Struct that hold the pending static data that will be pushed next tick. */
-	struct FPendingSubjectStatic
-	{
-		FLiveLinkSubjectKey SubjectKey;
-		TSubclassOf<ULiveLinkRole> Role;
-		FLiveLinkStaticDataStruct StaticData;
-	};
-
-	/** Struct that hold the pending frame data that will be pushed next tick. */
-	struct FPendingSubjectFrame
-	{
-		FLiveLinkSubjectKey SubjectKey;
-		FLiveLinkFrameDataStruct FrameData;
-	};
-
 private:
 	/** Common initialization code for the different constructors. */
 	void Initialize();
diff --git a/Engine/Plugins/Animation/LiveLink/Source/LiveLink/Public/LiveLinkSourceCollection.h b/Engine/Plugins/Animation/LiveLink/Source/LiveLink/Public/LiveLinkSourceCollection.h
--- a/Engine/Plugins/Animation/LiveLink/Source/LiveLink/Public/LiveLinkSourceCollection.h
+++ b/Engine/Plugins/Animation/LiveLink/Source/LiveLink/Public/LiveLinkSourceCollection.h
@@ -113,6 +113,7 @@ public:
 	void RemoveSubject(FLiveLinkSubjectKey SubjectKey);
 	FLiveLinkCollectionSubjectItem* FindSubject(FLiveLinkSubjectKey SubjectKey);
 	const FLiveLinkCollectionSubjectItem* FindSubject(FLiveLinkSubjectKey SubjectKey) const;
+	const FLiveLinkCollectionSubjectItem* FindSubject(FLiveLinkSubjectName SubjectName) const;
 	const FLiveLinkCollectionSubjectItem* FindEnabledSubject(FLiveLinkSubjectName SubjectName) const;
 	/** Get the number of subjects in the collection. */
 	int32 NumSubjects() const;
diff --git a/Engine/Plugins/Animation/LiveLink/Source/LiveLinkEditor/Private/LiveLinkClientPanelViews.cpp b/Engine/Plugins/Animation/LiveLink/Source/LiveLinkEditor/Private/LiveLinkClientPanelViews.cpp
--- a/Engine/Plugins/Animation/LiveLink/Source/LiveLinkEditor/Private/LiveLinkClientPanelViews.cpp
+++ b/Engine/Plugins/Animation/LiveLink/Source/LiveLinkEditor/Private/LiveLinkClientPanelViews.cpp
@@ -11,6 +11,7 @@
 #include "LiveLinkClient.h"
 #include "LiveLinkClientCommands.h"
 #include "LiveLinkSettings.h"
+#include "LiveLinkSubjectSettings.h"
 #include "LiveLinkTypes.h"
 #include "Framework/MultiBox/MultiBoxBuilder.h"
 #include "Modules/ModuleManager.h"
@@ -100,9 +101,10 @@ FText FLiveLinkSourceUIEntry::GetDisplayName() const
 	return GetSourceType();
 }
 
-FLiveLinkSubjectUIEntry::FLiveLinkSubjectUIEntry(const FLiveLinkSubjectKey& InSubjectKey, FLiveLinkClient* InClient)
+FLiveLinkSubjectUIEntry::FLiveLinkSubjectUIEntry(const FLiveLinkSubjectKey& InSubjectKey, FLiveLinkClient* InClient, bool bInIsSource)
 	: SubjectKey(InSubjectKey)
 	, Client(InClient)
+	, bIsSource(bInIsSource)
 {
 	if (InClient)
 	{
@@ -112,12 +114,12 @@ FLiveLinkSubjectUIEntry::FLiveLinkSubjectUIEntry(const FLiveLinkSubjectKey& InSu
 
 bool FLiveLinkSubjectUIEntry::IsSubject() const
 {
-	return !SubjectKey.SubjectName.IsNone();
+	return !bIsSource;
 }
 
 bool FLiveLinkSubjectUIEntry::IsSource() const
 {
-	return SubjectKey.SubjectName.IsNone();
+	return bIsSource;
 }
 
 bool FLiveLinkSubjectUIEntry::IsVirtualSubject() const
@@ -157,13 +159,13 @@ void FLiveLinkSubjectUIEntry::SetSubjectEnabled(bool bIsEnabled)
 
 FText FLiveLinkSubjectUIEntry::GetItemText() const
 {
-	if (IsSource())
+	if (IsSubject())
 	{
-		return Client->GetSourceType(SubjectKey.Source);
+		return Client->GetSubjectDisplayName(SubjectKey);
 	}
 	else
 	{
-		return FText::FromName(SubjectKey.SubjectName);
+		return Client->GetSourceNameOverride(SubjectKey);
 	}
 }
 
@@ -586,29 +588,32 @@ void FLiveLinkSubjectsView::RefreshSubjects()
 			TArray<FLiveLinkSubjectKey> SubjectKeys = Client->GetSubjects(true, true);
 			SubjectData.Reset();
 
-			TMap<FGuid, FLiveLinkSubjectUIEntryPtr> SourceHeaderItems;
+			TMap<FName, FLiveLinkSubjectUIEntryPtr> SourceItems;
+
 			TArray<FLiveLinkSubjectUIEntryPtr> AllItems;
 			AllItems.Reserve(SubjectKeys.Num());
 
 			for (const FLiveLinkSubjectKey& SubjectKey : SubjectKeys)
 			{
 				FLiveLinkSubjectUIEntryPtr Source;
-				if (FLiveLinkSubjectUIEntryPtr* SourcePtr = SourceHeaderItems.Find(SubjectKey.Source))
+
+				FName SourceNameOverride = *Client->GetSourceNameOverride(SubjectKey).ToString();
+
+				if (FLiveLinkSubjectUIEntryPtr* SourcePtr = SourceItems.Find(*SourceNameOverride.ToString()))
 				{
 					Source = *SourcePtr;
 				}
 				else
 				{
-					FLiveLinkSubjectKey SourceKey = SubjectKey;
-					SourceKey.SubjectName = NAME_None;
-					Source = MakeShared<FLiveLinkSubjectUIEntry>(SourceKey, static_cast<FLiveLinkClient*>(Client));
+					constexpr bool bIsSource = true;
+					Source = MakeShared<FLiveLinkSubjectUIEntry>(SubjectKey, static_cast<FLiveLinkClient*>(Client), bIsSource);
 					SubjectData.Add(Source);
-					SourceHeaderItems.Add(SubjectKey.Source) = Source;
+					SourceItems.Add(SourceNameOverride) = Source;
 
 					SubjectsTreeView->SetItemExpansion(Source, true);
 					AllItems.Add(Source);
 				}
-
+				
 				FLiveLinkSubjectUIEntryPtr SubjectEntry = MakeShared<FLiveLinkSubjectUIEntry>(SubjectKey, static_cast<FLiveLinkClient*>(Client));
 				Source->Children.Add(SubjectEntry);
 				AllItems.Add(SubjectEntry);
diff --git a/Engine/Plugins/Animation/LiveLink/Source/LiveLinkEditor/Public/LiveLinkClientPanelViews.h b/Engine/Plugins/Animation/LiveLink/Source/LiveLinkEditor/Public/LiveLinkClientPanelViews.h
--- a/Engine/Plugins/Animation/LiveLink/Source/LiveLinkEditor/Public/LiveLinkClientPanelViews.h
+++ b/Engine/Plugins/Animation/LiveLink/Source/LiveLinkEditor/Public/LiveLinkClientPanelViews.h
@@ -39,7 +39,7 @@ namespace UE::LiveLink
 // Structure that defines a single entry in the subject UI
 struct FLiveLinkSubjectUIEntry
 {
-	FLiveLinkSubjectUIEntry(const FLiveLinkSubjectKey& InSubjectKey, FLiveLinkClient* InClient);
+	FLiveLinkSubjectUIEntry(const FLiveLinkSubjectKey& InSubjectKey, FLiveLinkClient* InClient, bool bIsSource = false);
 
 	// Subject key
 	FLiveLinkSubjectKey SubjectKey;
@@ -73,6 +73,9 @@ struct FLiveLinkSubjectUIEntry
 private:
 	// Whether the subject is virtual
 	bool bIsVirtualSubject = false;
+
+	// Whether this represents a source.
+	bool bIsSource = false;
 };
 
 // Structure that defines a single entry in the source UI
diff --git a/Engine/Plugins/Animation/LiveLinkHub/Source/LiveLinkHubMessaging/Private/LiveLinkHubMessageBusSource.cpp b/Engine/Plugins/Animation/LiveLinkHub/Source/LiveLinkHubMessaging/Private/LiveLinkHubMessageBusSource.cpp
--- a/Engine/Plugins/Animation/LiveLinkHub/Source/LiveLinkHubMessaging/Private/LiveLinkHubMessageBusSource.cpp
+++ b/Engine/Plugins/Animation/LiveLinkHub/Source/LiveLinkHubMessaging/Private/LiveLinkHubMessageBusSource.cpp
@@ -8,6 +8,7 @@
 #include "Engine/World.h"
 #include "HAL/PlatformProcess.h"
 #include "ILiveLinkClient.h"
+#include "LiveLinkClient.h"
 #include "ILiveLinkModule.h"
 #include "LiveLinkHubMessages.h"
 #include "LiveLinkMessages.h"
@@ -71,6 +72,22 @@ void FLiveLinkHubMessageBusSource::OnMapChanged(UWorld* World, EMapChangeType Ch
 	SendClientInfoMessage();
 }
 
+void FLiveLinkHubMessageBusSource::InitializeAndPushStaticData_AnyThread(FName SubjectName, TSubclassOf<ULiveLinkRole> SubjectRole, const FLiveLinkSubjectKey& SubjectKey, const TSharedRef<IMessageContext, ESPMode::ThreadSafe>& Context, UScriptStruct* MessageTypeInfo)
+{
+	check(MessageTypeInfo->IsChildOf(FLiveLinkBaseStaticData::StaticStruct()));
+
+	FLiveLinkStaticDataStruct DataStruct(MessageTypeInfo);
+	DataStruct.InitializeWith(MessageTypeInfo, reinterpret_cast<const FLiveLinkBaseStaticData*>(Context->GetMessage()));
+
+	FLiveLinkClient::FPendingSubjectStatic PendingStaticData;
+	PendingStaticData.Role = SubjectRole;
+	PendingStaticData.SubjectKey = SubjectKey;
+	PendingStaticData.StaticData = MoveTemp(DataStruct);
+	PendingStaticData.ExtraMetadata = Context->GetAnnotations();
+
+	static_cast<FLiveLinkClient*>(Client)->PushPendingSubject_AnyThread(MoveTemp(PendingStaticData));
+}
+
 void FLiveLinkHubMessageBusSource::InitializeMessageEndpoint(FMessageEndpointBuilder& EndpointBuilder)
 {
 	FLiveLinkMessageBusSource::InitializeMessageEndpoint(EndpointBuilder);
diff --git a/Engine/Plugins/Animation/LiveLinkHub/Source/LiveLinkHubMessaging/Private/LiveLinkHubMessageBusSource.h b/Engine/Plugins/Animation/LiveLinkHub/Source/LiveLinkHubMessaging/Private/LiveLinkHubMessageBusSource.h
--- a/Engine/Plugins/Animation/LiveLinkHub/Source/LiveLinkHubMessaging/Private/LiveLinkHubMessageBusSource.h
+++ b/Engine/Plugins/Animation/LiveLinkHub/Source/LiveLinkHubMessaging/Private/LiveLinkHubMessageBusSource.h
@@ -19,6 +19,7 @@ public:
 
 protected:
 	//~ Begin FLiveLinkMessageBusSource interface
+	virtual void InitializeAndPushStaticData_AnyThread(FName SubjectName, TSubclassOf<ULiveLinkRole> SubjectRole, const FLiveLinkSubjectKey& SubjectKey, const TSharedRef<IMessageContext, ESPMode::ThreadSafe>& Context, UScriptStruct* MessageTypeInfo) override;
 	virtual void InitializeMessageEndpoint(FMessageEndpointBuilder& EndpointBuilder);
 	virtual double GetDeadSourceTimeout() const override;
 	virtual void SendConnectMessage() override;
diff --git a/Engine/Source/Runtime/LiveLinkInterface/Public/ILiveLinkClient.h b/Engine/Source/Runtime/LiveLinkInterface/Public/ILiveLinkClient.h
--- a/Engine/Source/Runtime/LiveLinkInterface/Public/ILiveLinkClient.h
+++ b/Engine/Source/Runtime/LiveLinkInterface/Public/ILiveLinkClient.h
@@ -302,6 +302,12 @@ public:
 	 */
 	virtual UObject* GetSubjectSettings(const FLiveLinkSubjectKey& SubjectKey) const = 0;
 
+	/** Get the source name override for a given subject. Allows subjects to modify the source's display name in the UI when needed. */
+	virtual FText GetSourceNameOverride(const FLiveLinkSubjectKey& SubjectKey) const = 0;
+
+	/** Get the display name for a subject. Returns an empty text if the source is invalid. */
+	virtual FText GetSubjectDisplayName(const FLiveLinkSubjectKey& SubjectKey) const = 0;
+
 	/**
 	 * Utility method to grab a subject's static data. Used by the RecordingController when static data is missing from the recording. 
 	 */
diff --git a/Engine/Source/Runtime/LiveLinkInterface/Public/LiveLinkSubjectSettings.h b/Engine/Source/Runtime/LiveLinkInterface/Public/LiveLinkSubjectSettings.h
--- a/Engine/Source/Runtime/LiveLinkInterface/Public/LiveLinkSubjectSettings.h
+++ b/Engine/Source/Runtime/LiveLinkInterface/Public/LiveLinkSubjectSettings.h
@@ -67,6 +67,11 @@ public:
 		return Key.SubjectName.Name;
 	}
 
+	virtual FText GetDisplayName() const
+	{
+		return FText::FromName(Key.SubjectName);
+	}
+
 	/** List of available preprocessor the subject will use. */
 	UPROPERTY(EditAnywhere, Instanced, Category = "LiveLink", meta = (DisplayName = "Pre Processors"))
 	TArray<TObjectPtr<ULiveLinkFramePreProcessor>> PreProcessors;
@@ -98,6 +103,10 @@ public:
 	UPROPERTY(EditAnywhere, Category = "LiveLink", meta=(EditCondition="bAllowModifyingRebroadcast"))
     bool bRebroadcastSubject;
 
+	/** For sources created through LiveLinkHub, this contains the name of the original source for display purposes.*/
+	UPROPERTY()
+	FName OriginalSourceName;
+
 	/** Validate PreProcessors, Translators and Interpolation processors. Usually called after a property change event.
 	 * Will revert a given change if it does not match the current subject role.
 	 */
diff --git a/Engine/Source/Runtime/LiveLinkInterface/Public/LiveLinkVirtualSubject.h b/Engine/Source/Runtime/LiveLinkInterface/Public/LiveLinkVirtualSubject.h
--- a/Engine/Source/Runtime/LiveLinkInterface/Public/LiveLinkVirtualSubject.h
+++ b/Engine/Source/Runtime/LiveLinkInterface/Public/LiveLinkVirtualSubject.h
@@ -70,6 +70,12 @@ public:
 	/** Returns true whether this virtual subject depends on the Subject named SubjectName */
 	LIVELINKINTERFACE_API virtual bool DependsOnSubject(FName SubjectName) const;
 	
+	/** Get display name for this subject. */
+	virtual FText GetDisplayName() const
+	{
+		return FText::FromName(SubjectKey.SubjectName);
+	}
+	
 protected:
 
 	/** Updates the list of translators valid for this frame */
diff --git a/Engine/Source/Runtime/LiveLinkMessageBusFramework/Private/LiveLinkProvider.cpp b/Engine/Source/Runtime/LiveLinkMessageBusFramework/Private/LiveLinkProvider.cpp
--- a/Engine/Source/Runtime/LiveLinkMessageBusFramework/Private/LiveLinkProvider.cpp
+++ b/Engine/Source/Runtime/LiveLinkMessageBusFramework/Private/LiveLinkProvider.cpp
@@ -24,6 +24,7 @@ template TSharedPtr<ILiveLinkProvider> ILiveLinkProvider::CreateLiveLinkProvider
 
 FName FLiveLinkMessageAnnotation::SubjectAnnotation = TEXT("SubjectName");
 FName FLiveLinkMessageAnnotation::RoleAnnotation = TEXT("Role");
+FName FLiveLinkMessageAnnotation::OriginalSourceAnnotation = TEXT("OriginalSource");
 
 
 // Address that we have had a connection request from
@@ -64,12 +65,15 @@ struct FTrackedStaticData
 	FTrackedStaticData()
 		: SubjectName(NAME_None)
 	{}
-	FTrackedStaticData(FName InSubjectName, TWeakObjectPtr<UClass> InRoleClass, FLiveLinkStaticDataStruct InStaticData)
+	FTrackedStaticData(FName InSubjectName, TWeakObjectPtr<UClass> InRoleClass, FLiveLinkStaticDataStruct InStaticData, TMap<FName, FString>&& InAnnotations)
 		: SubjectName(InSubjectName), RoleClass(InRoleClass), StaticData(MoveTemp(InStaticData))
+		, Annotations(MoveTemp(InAnnotations))
 	{}
 	FName SubjectName;
 	TWeakObjectPtr<UClass> RoleClass;
 	FLiveLinkStaticDataStruct StaticData;
+	TMap<FName, FString> Annotations;
+
 	bool operator==(FName InSubjectName) const { return SubjectName == InSubjectName; }
 };
 
@@ -80,11 +84,13 @@ struct FTrackedFrameData
 	FTrackedFrameData()
 		: SubjectName(NAME_None)
 	{}
-	FTrackedFrameData(FName InSubjectName, FLiveLinkFrameDataStruct InFrameData)
+	FTrackedFrameData(FName InSubjectName, FLiveLinkFrameDataStruct InFrameData, TMap<FName, FString>&& InAnnotations)
 		: SubjectName(InSubjectName), FrameData(MoveTemp(InFrameData))
+		, Annotations(MoveTemp(InAnnotations))
 	{}
 	FName SubjectName;
 	FLiveLinkFrameDataStruct FrameData;
+	TMap<FName, FString> Annotations;
 	bool operator==(FName InSubjectName) const { return SubjectName == InSubjectName; }
 };
 
@@ -203,30 +209,32 @@ FTrackedFrameData* FLiveLinkProvider::GetLastSubjectFrameData(const FName& Subje
 	return FrameDatas.FindByKey(SubjectName);
 }
 
-void FLiveLinkProvider::SetLastSubjectStaticData(FName SubjectName, TSubclassOf<ULiveLinkRole> Role, FLiveLinkStaticDataStruct&& StaticData)
+void FLiveLinkProvider::SetLastSubjectStaticData(FName SubjectName, TSubclassOf<ULiveLinkRole> Role, FLiveLinkStaticDataStruct&& StaticData, TMap<FName, FString>&& SubjectAnnotations)
 {
 	FTrackedStaticData* Result = StaticDatas.FindByKey(SubjectName);
 	if (Result)
 	{
+		Result->Annotations = MoveTemp(SubjectAnnotations);
 		Result->StaticData = MoveTemp(StaticData);
 		Result->RoleClass = Role.Get();
 	}
 	else
 	{
-		StaticDatas.Emplace(SubjectName, Role.Get(), MoveTemp(StaticData));
+		StaticDatas.Emplace(SubjectName, Role.Get(), MoveTemp(StaticData), MoveTemp(SubjectAnnotations));
 	}
 }
 
-void FLiveLinkProvider::SetLastSubjectFrameData(FName SubjectName, FLiveLinkFrameDataStruct&& FrameData)
+void FLiveLinkProvider::SetLastSubjectFrameData(FName SubjectName, FLiveLinkFrameDataStruct&& FrameData, TMap<FName, FString>&& SubjectAnnotations)
 {
 	FTrackedFrameData* Result = FrameDatas.FindByKey(SubjectName);
 	if (Result)
 	{
 		Result->FrameData = MoveTemp(FrameData);
+		Result->Annotations = MoveTemp(SubjectAnnotations);
 	}
 	else
 	{
-		FrameDatas.Emplace(SubjectName, MoveTemp(FrameData));
+		FrameDatas.Emplace(SubjectName, MoveTemp(FrameData), MoveTemp(SubjectAnnotations));
 	}
 }
 
@@ -302,7 +310,7 @@ void FLiveLinkProvider::SendClearSubjectToConnections(FName SubjectName)
 	MessageEndpoint->Send(FMessageEndpoint::MakeMessage<FLiveLinkClearSubject>(SubjectName), EMessageFlags::Reliable, GetAnnotations(), nullptr, MessageAddresses, FTimespan::Zero(), FDateTime::MaxValue());
 }
 
-bool FLiveLinkProvider::UpdateSubjectStaticData(const FName SubjectName, TSubclassOf<ULiveLinkRole> Role, FLiveLinkStaticDataStruct&& StaticData)
+bool FLiveLinkProvider::UpdateSubjectStaticData(const FName SubjectName, TSubclassOf<ULiveLinkRole> Role, FLiveLinkStaticDataStruct&& StaticData, const TMap<FName, FString>& ExtraAnnotations)
 {
 	FScopeLock Lock(&CriticalSection);
 
@@ -323,19 +331,20 @@ bool FLiveLinkProvider::UpdateSubjectStaticData(const FName SubjectName, TSubcla
 
 	ValidateConnections();
 
+	TMap<FName, FString> Annotations;
+	Annotations.Add(FLiveLinkMessageAnnotation::SubjectAnnotation, SubjectName.ToString());
+	Annotations.Add(FLiveLinkMessageAnnotation::RoleAnnotation, Role->GetName());
+	Annotations.Append(ExtraAnnotations);
+
 	if (ConnectedAddresses.Num() > 0)
 	{
 		TArray<FMessageAddress> Addresses;
 		GetFilteredAddresses(SubjectName, Addresses);
 
-		TMap<FName, FString> Annotations;
-		Annotations.Add(FLiveLinkMessageAnnotation::SubjectAnnotation, SubjectName.ToString());
-		Annotations.Add(FLiveLinkMessageAnnotation::RoleAnnotation, Role->GetName());
-
 		MessageEndpoint->Send(StaticData.CloneData(), const_cast<UScriptStruct*>(StaticData.GetStruct()), EMessageFlags::Reliable, Annotations, nullptr, Addresses, FTimespan::Zero(), FDateTime::MaxValue());
 	}
 
-	SetLastSubjectStaticData(SubjectName, Role, MoveTemp(StaticData));
+	SetLastSubjectStaticData(SubjectName, Role, MoveTemp(StaticData), MoveTemp(Annotations));
 
 	return true;
 }
@@ -383,7 +392,7 @@ void FLiveLinkProvider::UpdateSubjectFrame(const FName& SubjectName, const TArra
 	SendSubjectFrame(SubjectName, Subject);
 }
 
-bool FLiveLinkProvider::UpdateSubjectFrameData(const FName SubjectName, FLiveLinkFrameDataStruct&& FrameData)
+bool FLiveLinkProvider::UpdateSubjectFrameData(const FName SubjectName, FLiveLinkFrameDataStruct&& FrameData, const TMap<FName, FString>& ExtraAnnotations)
 {
 	FScopeLock Lock(&CriticalSection);
 
@@ -411,18 +420,19 @@ bool FLiveLinkProvider::UpdateSubjectFrameData(const FName SubjectName, FLiveLin
 
 	ValidateConnections();
 
+	TMap<FName, FString> Annotations;
+	Annotations.Add(FLiveLinkMessageAnnotation::SubjectAnnotation, SubjectName.ToString());
+	Annotations.Append(ExtraAnnotations);
+
 	if (ConnectedAddresses.Num() > 0)
 	{
 		TArray<FMessageAddress> Addresses;
 		GetFilteredAddresses(SubjectName, Addresses);
 
-		TMap<FName, FString> Annotations;
-		Annotations.Add(FLiveLinkMessageAnnotation::SubjectAnnotation, SubjectName.ToString());
-
 		MessageEndpoint->Send(FrameData.CloneData(), const_cast<UScriptStruct*>(FrameData.GetStruct()), EMessageFlags::None, Annotations, nullptr, Addresses, FTimespan::Zero(), FDateTime::MaxValue());
 	}
 
-	SetLastSubjectFrameData(SubjectName, MoveTemp(FrameData));
+	SetLastSubjectFrameData(SubjectName, MoveTemp(FrameData), MoveTemp(Annotations));
 
 	return true;
 }
@@ -501,8 +511,8 @@ void FLiveLinkProvider::HandleConnectMessage(const FLiveLinkConnectMessage& Mess
 		for (const FTrackedStaticData& Data : StaticDatas)
 		{
 			UClass* RoleClass = Data.RoleClass.Get();
-			Annotations.FindChecked(FLiveLinkMessageAnnotation::SubjectAnnotation) = Data.SubjectName.ToString();
-			Annotations.FindChecked(FLiveLinkMessageAnnotation::RoleAnnotation) = RoleClass ? RoleClass->GetName() : TEXT("");
+			Annotations.Append(Data.Annotations);
+
 			MessageEndpoint->Send(Data.StaticData.CloneData(), const_cast<UScriptStruct*>(Data.StaticData.GetStruct()), EMessageFlags::Reliable, Annotations, nullptr, MessageAddress, FTimespan::Zero(), FDateTime::MaxValue());
 		}
 
@@ -510,7 +520,7 @@ void FLiveLinkProvider::HandleConnectMessage(const FLiveLinkConnectMessage& Mess
 
 		for (const FTrackedFrameData& Data : FrameDatas)
 		{
-			Annotations.FindChecked(FLiveLinkMessageAnnotation::SubjectAnnotation) = Data.SubjectName.ToString();
+			Annotations.Append(Data.Annotations);
 			MessageEndpoint->Send(Data.FrameData.CloneData(), const_cast<UScriptStruct*>(Data.FrameData.GetStruct()), EMessageFlags::None, Annotations, nullptr, MessageAddress, FTimespan::Zero(), FDateTime::MaxValue());
 		}
 
diff --git a/Engine/Source/Runtime/LiveLinkMessageBusFramework/Public/LiveLinkMessages.h b/Engine/Source/Runtime/LiveLinkMessageBusFramework/Public/LiveLinkMessages.h
--- a/Engine/Source/Runtime/LiveLinkMessageBusFramework/Public/LiveLinkMessages.h
+++ b/Engine/Source/Runtime/LiveLinkMessageBusFramework/Public/LiveLinkMessages.h
@@ -10,6 +10,7 @@ struct LIVELINKMESSAGEBUSFRAMEWORK_API FLiveLinkMessageAnnotation
 {
 	static FName SubjectAnnotation;
 	static FName RoleAnnotation;
+	static FName OriginalSourceAnnotation;
 };
 
 USTRUCT()
diff --git a/Engine/Source/Runtime/LiveLinkMessageBusFramework/Public/LiveLinkProvider.h b/Engine/Source/Runtime/LiveLinkMessageBusFramework/Public/LiveLinkProvider.h
--- a/Engine/Source/Runtime/LiveLinkMessageBusFramework/Public/LiveLinkProvider.h
+++ b/Engine/Source/Runtime/LiveLinkMessageBusFramework/Public/LiveLinkProvider.h
@@ -44,15 +44,16 @@ public:
 
 	/**
 	 * Send, to UE, the static data of a subject.
-	 * @param SubjectName	The name of the subject
-	 * @param Role			The Live Link role of the subject. The StaticData type should match the role's data.
-	 * @param StaticData	The static data of the subject.
-							The FLiveLinkStaticDataStruct doesn't have a copy constructor.
-	 *						The argument is passed by r-value to help the user understand the compiler error message.
+	 * @param SubjectName	    The name of the subject
+	 * @param Role			    The Live Link role of the subject. The StaticData type should match the role's data.
+	 * @param StaticData	    The static data of the subject.
+							    The FLiveLinkStaticDataStruct doesn't have a copy constructor.
+	 *						    The argument is passed by r-value to help the user understand the compiler error message.
+	 * @param ExtraAnnotations	Extra annotations to pass in with the static data message.
 	 * @return				True if the message was sent or is pending an active connection.
 	 * @see					UpdateSubjectFrameData, RemoveSubject
 	 */
-	virtual bool UpdateSubjectStaticData(const FName SubjectName, TSubclassOf<ULiveLinkRole> Role, FLiveLinkStaticDataStruct&& StaticData) = 0;
+	virtual bool UpdateSubjectStaticData(const FName SubjectName, TSubclassOf<ULiveLinkRole> Role, FLiveLinkStaticDataStruct&& StaticData, const TMap<FName, FString>& ExtraAnnotations = {}) = 0;
 
 	/**
 	 * Inform UE that a subject won't be streamed anymore.
@@ -62,14 +63,15 @@ public:
 
 	/**
 	 * Send the static data of a subject to UE.
-	 * @param SubjectName	The name of the subject
-	 * @param StaticData	The frame data of the subject. The type should match the role's data send with UpdateSubjectStaticData.
-							The FLiveLinkFrameDataStruct doesn't have a copy constructor.
-	 *						The argument is passed by r-value to help the user understand the compiler error message.
-	 * @return				True if the message was sent or is pending an active connection.
-	 * @see					UpdateSubjectStaticData, RemoveSubject
+	 * @param SubjectName		The name of the subject
+	 * @param StaticData		The frame data of the subject. The type should match the role's data send with UpdateSubjectStaticData.
+								The FLiveLinkFrameDataStruct doesn't have a copy constructor.
+	 *							The argument is passed by r-value to help the user understand the compiler error message.
+	 * @param ExtraAnnotations	Extra annotations to pass in with the frame data message.
+	 * @return					True if the message was sent or is pending an active connection.
+	 * @see						UpdateSubjectStaticData, RemoveSubject
 	 */
-	virtual bool UpdateSubjectFrameData(const FName SubjectName, FLiveLinkFrameDataStruct&& FrameData) = 0;
+	virtual bool UpdateSubjectFrameData(const FName SubjectName, FLiveLinkFrameDataStruct&& FrameData, const TMap<FName, FString>& ExtraAnnotations = {}) = 0;
 
 	/** Is this provider currently connected to something. */
 	virtual bool HasConnection() const = 0;
diff --git a/Engine/Source/Runtime/LiveLinkMessageBusFramework/Public/LiveLinkProviderImpl.h b/Engine/Source/Runtime/LiveLinkMessageBusFramework/Public/LiveLinkProviderImpl.h
--- a/Engine/Source/Runtime/LiveLinkMessageBusFramework/Public/LiveLinkProviderImpl.h
+++ b/Engine/Source/Runtime/LiveLinkMessageBusFramework/Public/LiveLinkProviderImpl.h
@@ -70,9 +70,9 @@ private:
 
 	FTrackedFrameData* GetLastSubjectFrameData(const FName& SubjectName);
 
-	void SetLastSubjectStaticData(FName SubjectName, TSubclassOf<ULiveLinkRole> Role, FLiveLinkStaticDataStruct&& StaticData);
+	void SetLastSubjectStaticData(FName SubjectName, TSubclassOf<ULiveLinkRole> Role, FLiveLinkStaticDataStruct&& StaticData, TMap<FName, FString>&& Annotations);
 
-	void SetLastSubjectFrameData(FName SubjectName, FLiveLinkFrameDataStruct&& FrameData);
+	void SetLastSubjectFrameData(FName SubjectName, FLiveLinkFrameDataStruct&& FrameData, TMap<FName, FString>&& SubjectAnnotations);
 
 	// Clear a existing track subject
 	void ClearTrackedSubject(const FName& SubjectName);
@@ -186,7 +186,7 @@ public:
 
 	virtual void UpdateSubject(const FName& SubjectName, const TArray<FName>& BoneNames, const TArray<int32>& BoneParents);
 
-	virtual bool UpdateSubjectStaticData(const FName SubjectName, TSubclassOf<ULiveLinkRole> Role, FLiveLinkStaticDataStruct&& StaticData) override;
+	virtual bool UpdateSubjectStaticData(const FName SubjectName, TSubclassOf<ULiveLinkRole> Role, FLiveLinkStaticDataStruct&& StaticData, const TMap<FName, FString>& ExtraAnnotations = {}) override;
 
 	virtual void ClearSubject(const FName& SubjectName);
 
@@ -199,8 +199,8 @@ PRAGMA_DISABLE_DEPRECATION_WARNINGS
 									const FLiveLinkMetaData& MetaData, double Time);
 PRAGMA_ENABLE_DEPRECATION_WARNINGS
 
-	virtual bool UpdateSubjectFrameData(const FName SubjectName, FLiveLinkFrameDataStruct&& FrameData) override;
-	
+	virtual bool UpdateSubjectFrameData(const FName SubjectName, FLiveLinkFrameDataStruct&& FrameData, const TMap<FName, FString>& ExtraAnnotations = {}) override;
+
 	virtual bool HasConnection() const override;
 
 	virtual FDelegateHandle RegisterConnStatusChangedHandle(const FLiveLinkProviderConnectionStatusChanged::FDelegate& ConnStatusChanged) override;
