LiveLink - Add structures necessary for implementing payload compression

#rnx

diff --git a/Engine/Source/Runtime/LiveLinkInterface/Public/Roles/LiveLinkAnimationTypes.h b/Engine/Source/Runtime/LiveLinkInterface/Public/Roles/LiveLinkAnimationTypes.h
--- a/Engine/Source/Runtime/LiveLinkInterface/Public/Roles/LiveLinkAnimationTypes.h
+++ b/Engine/Source/Runtime/LiveLinkInterface/Public/Roles/LiveLinkAnimationTypes.h
@@ -57,3 +57,41 @@ public:
 	TArray<FTransform> Transforms;
 };
 
+/** Single-precision floating point equivalent of FLiveLinkAnimationFrameData. */
+USTRUCT()
+struct FLiveLinkFloatAnimationFrameData : public FLiveLinkBaseFrameData
+{
+	GENERATED_BODY();
+
+	// Converts double-precision float animation data to single-precision float animation data.
+	static FLiveLinkFloatAnimationFrameData FromAnimData(const FLiveLinkAnimationFrameData& InAnimData)
+	{
+		FLiveLinkFloatAnimationFrameData FloatData;
+
+		const FLiveLinkBaseFrameData* BaseFrameData = &InAnimData;
+		*dynamic_cast<FLiveLinkBaseFrameData*>(&FloatData) = *BaseFrameData;
+
+		FloatData.Transforms.Reserve(InAnimData.Transforms.Num());
+		Algo::Transform(InAnimData.Transforms, FloatData.Transforms, [](const FTransform& Transform) { return FTransform3f(Transform); });
+
+		return FloatData;
+	}
+
+	// Converts single-precision float animation data to double-precision float animation data.
+	static FLiveLinkAnimationFrameData ToAnimData(const FLiveLinkFloatAnimationFrameData& InFloatData)
+	{
+		FLiveLinkAnimationFrameData AnimData;
+
+		const FLiveLinkBaseFrameData* BaseFrameData = &InFloatData;
+		*dynamic_cast<FLiveLinkBaseFrameData*>(&AnimData) = *BaseFrameData;
+
+		AnimData.Transforms.Reserve(InFloatData.Transforms.Num());
+		Algo::Transform(InFloatData.Transforms, AnimData.Transforms, [](const FTransform3f& Transform) { return FTransform(Transform); });
+
+		return AnimData;
+	}
+
+	// Array of transforms for each bone of the skeleton
+	UPROPERTY()
+	TArray<FTransform3f> Transforms;
+};
diff --git a/Engine/Source/Runtime/LiveLinkMessageBusFramework/LiveLinkMessageBusFramework.Build.cs b/Engine/Source/Runtime/LiveLinkMessageBusFramework/LiveLinkMessageBusFramework.Build.cs
--- a/Engine/Source/Runtime/LiveLinkMessageBusFramework/LiveLinkMessageBusFramework.Build.cs
+++ b/Engine/Source/Runtime/LiveLinkMessageBusFramework/LiveLinkMessageBusFramework.Build.cs
@@ -12,6 +12,7 @@ public class LiveLinkMessageBusFramework : ModuleRules
 			new string[] {
 				"Core",
 				"CoreUObject",
+				"Serialization",
 				"LiveLinkInterface",
 				"MessagingCommon"
 			}
diff --git a/Engine/Source/Runtime/LiveLinkMessageBusFramework/Private/LiveLinkCompression.cpp b/Engine/Source/Runtime/LiveLinkMessageBusFramework/Private/LiveLinkCompression.cpp
new file mode 100644
--- /dev/null
+++ b/Engine/Source/Runtime/LiveLinkMessageBusFramework/Private/LiveLinkCompression.cpp
@@ -0,0 +1,231 @@
+//Copyright Epic Games, Inc. All Rights Reserved.
+
+#include "LiveLinkCompression.h"
+
+
+#include "Misc/App.h"
+#include "Misc/Compression.h"
+#include "UObject/StructOnScope.h"
+
+#include "HAL/IConsoleManager.h"
+#include "Serialization/MemoryReader.h"
+#include "Serialization/MemoryWriter.h"
+#include "StructSerializer.h"
+#include "StructDeserializer.h"
+#include "Backends/CborStructSerializerBackend.h"
+#include "Backends/CborStructDeserializerBackend.h"
+
+namespace UE::LiveLink::Compression
+{
+	static TAutoConsoleVariable<int32> CVarCompressionType(
+		TEXT("LiveLink.SetCompressionType"), 1,
+		TEXT("Specify the type of compression to use when serializing data. A value of 0 means compression is off. A value of 1 = Oodle. All other values = Zlib."));
+
+	static TAutoConsoleVariable<int32> CVarCompressionFlags(
+		TEXT("LiveLink.SetCompressionFlags"), 0,
+		TEXT("Specify the flags to use when compression is enabled. A value of 0 means no flags. A value of 1 favors smaller sizes. Any other value favors faster encoding."));
+
+	int32 GetConsoleVariableCompressionType()
+	{
+		return CVarCompressionType.GetValueOnAnyThread();
+	}
+
+	int32 GetConsoleVariableCompressionFlags()
+	{
+		return CVarCompressionFlags.GetValueOnAnyThread();
+	}
+
+}
+
+namespace PayloadDetail
+{
+
+	bool ShouldCompress(const FLiveLinkSerializedFrameData& InPayload, ELiveLinkPayloadCompressionType CompressionType)
+	{
+		if (CompressionType == ELiveLinkPayloadCompressionType::None)
+		{
+			return false;
+		}
+		if (CompressionType == ELiveLinkPayloadCompressionType::Heuristic)
+		{
+			return UE::LiveLink::Compression::ShouldCompress(InPayload.PayloadSize);
+		}
+		check(CompressionType == ELiveLinkPayloadCompressionType::Always);
+		// Otherwise we are always compressing
+		return InPayload.PayloadSize > 0;
+	}
+
+	bool TryCompressImpl(const UScriptStruct* InEventType, const void* InEventData, FLiveLinkSerializedFrameData& InOutPayload, ELiveLinkPayloadCompressionType CompressionType)
+	{
+		InOutPayload.PayloadSize = InOutPayload.PayloadBytes.Bytes.Num();
+
+		// if we serialized something, compress it
+		if (ShouldCompress(InOutPayload, CompressionType))
+		{
+			TRACE_CPUPROFILER_EVENT_SCOPE(LiveLink::TryCompressImpl);
+
+			TArray<uint8>& InBytes = InOutPayload.PayloadBytes.Bytes;
+			TArray<uint8> OutCompressedData;
+			// Compress the result to send on the wire
+			FName NamedCompressionAlgo = UE::LiveLink::Compression::GetCompressionAlgorithm();
+			ECompressionFlags CompressFlags = UE::LiveLink::Compression::GetCompressionFlags();
+			int32 CompressedSize = FCompression::CompressMemoryBound(NamedCompressionAlgo, InOutPayload.PayloadSize, CompressFlags);
+			OutCompressedData.SetNumUninitialized(CompressedSize);
+
+			if (FCompression::CompressMemory(NamedCompressionAlgo, OutCompressedData.GetData(), CompressedSize, InBytes.GetData(), InBytes.Num(), CompressFlags))
+			{
+				OutCompressedData.SetNum(CompressedSize, EAllowShrinking::No);
+				InOutPayload.PayloadBytes.Bytes = MoveTemp(OutCompressedData);
+
+				InOutPayload.CompressionMethod = UE::LiveLink::Compression::GetCompressionMethod(NamedCompressionAlgo);
+				InOutPayload.CompressionBias = UE::LiveLink::Compression::GetCompressionBias(CompressFlags);
+			}
+			else
+			{
+				UE_LOG(LogTemp, Warning, TEXT("Unable to compress data for %s!"), *InEventType->GetName());
+				InOutPayload.CompressionMethod = ELiveLinkCompressionMethod::Uncompressed;
+			}
+		}
+		else
+		{
+			InOutPayload.CompressionMethod = ELiveLinkCompressionMethod::Uncompressed;
+		}
+
+		// Since we can support uncompressed or compressed data this is always successful.
+		return true;
+	}
+
+	static TOptional<TArray<uint8>> DecompressImpl(const FLiveLinkSerializedFrameData& InPayload)
+	{
+		ELiveLinkCompressionMethod CompressMethod = InPayload.CompressionMethod;
+		if (CompressMethod != ELiveLinkCompressionMethod::Uncompressed)
+		{
+			TRACE_CPUPROFILER_EVENT_SCOPE(LiveLink::TryDecompress);
+
+			const TArray<uint8>& InBytes = InPayload.PayloadBytes.Bytes;
+			TArray<uint8> UncompressedData;
+			UncompressedData.SetNumUninitialized(InPayload.PayloadSize);
+
+			ECompressionFlags CompressFlags = UE::LiveLink::Compression::GetCoreCompressionFlags(InPayload.CompressionBias);
+			FName CompressType = UE::LiveLink::Compression::GetCompressionAlgorithm(InPayload.CompressionMethod);
+			if (FCompression::UncompressMemory(CompressType, UncompressedData.GetData(), UncompressedData.Num(), InBytes.GetData(), InBytes.Num(), CompressFlags))
+			{
+				return TOptional<TArray<uint8>>(MoveTemp(UncompressedData));
+			}
+			else
+			{
+				UE_LOG(LogTemp, Warning, TEXT("Unable to uncompress data for %s!"), *InPayload.PayloadTypeName.ToString());
+			}
+		}
+
+		return TOptional<TArray<uint8>>{};
+	}
+
+	bool SerializeImpl(const UScriptStruct* InSourceEventType, const void* InSourceEventData, FLiveLinkSerializedFrameData& OutSerializedData)
+	{
+		if (OutSerializedData.SerializationMethod == ELiveLinkPayloadSerializationMethod::Cbor)
+		{
+			FMemoryWriter Writer(OutSerializedData.PayloadBytes.Bytes);
+			FCborStructSerializerBackend Serializer(Writer, EStructSerializerBackendFlags::Default);
+			FStructSerializer::Serialize(InSourceEventData, *const_cast<UScriptStruct*>(InSourceEventType), Serializer);
+			return !Writer.GetError();
+		}
+
+		FMemoryWriter Archive(OutSerializedData.PayloadBytes.Bytes);
+		Archive.SetWantBinaryPropertySerialization(true);
+		const_cast<UScriptStruct*>(InSourceEventType)->SerializeItem(Archive, (uint8*)InSourceEventData, nullptr);
+		return !Archive.GetError();
+	}
+
+	bool DeserializeImpl(const UScriptStruct* InTargetEventType, void* InOutTargetEventData, ELiveLinkPayloadSerializationMethod SerializeMethod, const TArray<uint8>& InBytes)
+	{
+		if (SerializeMethod == ELiveLinkPayloadSerializationMethod::Cbor)
+		{
+			FMemoryReader Reader(InBytes);
+			FCborStructDeserializerBackend Deserializer(Reader);
+			return FStructDeserializer::Deserialize(InOutTargetEventData, *const_cast<UScriptStruct*>(InTargetEventType), Deserializer) && !Reader.GetError();
+		}
+
+		FMemoryReader Archive(InBytes);
+		Archive.SetWantBinaryPropertySerialization(true);
+		const_cast<UScriptStruct*>(InTargetEventType)->SerializeItem(Archive, (uint8*)InOutTargetEventData, nullptr);
+		return !Archive.GetError();
+	}
+
+	bool DeserializeAndDecompress(const UScriptStruct* InTargetEventType, void* InOutTargetEventData, const FLiveLinkSerializedFrameData& InPayload)
+	{
+		TOptional<TArray<uint8>> DecompressedBytes = DecompressImpl(InPayload);
+		ELiveLinkCompressionMethod CompressMethod = InPayload.CompressionMethod;
+
+		if (CompressMethod != ELiveLinkCompressionMethod::Uncompressed && !DecompressedBytes.IsSet())
+		{
+			return false;
+		}
+		const TArray<uint8>& ByteStream = DecompressedBytes.IsSet() ? DecompressedBytes.GetValue() : InPayload.PayloadBytes.Bytes;
+		return DeserializeImpl(InTargetEventType, InOutTargetEventData, InPayload.SerializationMethod, ByteStream);
+	}
+
+	bool SerializePreChecks(const UScriptStruct* InSourceEventType, const void* InSourceEventData, FLiveLinkSerializedFrameData& OutSerializedData)
+	{
+		OutSerializedData.PayloadSize = 0;
+		OutSerializedData.PayloadBytes.Bytes.Reset();
+
+		return InSourceEventType && InSourceEventData;
+	}
+
+	bool DeserializePreChecks(const UScriptStruct* InEventType, void* InOutEventData, const FLiveLinkSerializedFrameData& Payload)
+	{
+		return InEventType && InOutEventData;
+	}
+
+} // namespace PayloadDetail
+
+bool FLiveLinkSerializedFrameData::SetPayload(const FStructOnScope& InPayload, ELiveLinkPayloadCompressionType CompressionType)
+{
+	const UStruct* PayloadStruct = InPayload.GetStruct();
+	check(PayloadStruct->IsA<UScriptStruct>());
+	return SetPayload((UScriptStruct*)PayloadStruct, InPayload.GetStructMemory(), CompressionType);
+}
+
+bool FLiveLinkSerializedFrameData::SetPayload(const UScriptStruct* InPayloadType, const void* InPayloadData, ELiveLinkPayloadCompressionType CompressionType)
+{
+	check(InPayloadType && InPayloadData);
+	PayloadTypeName = *InPayloadType->GetPathName();
+	return PayloadDetail::SerializePreChecks(InPayloadType, InPayloadData, *this)
+		&& PayloadDetail::SerializeImpl(InPayloadType, InPayloadData, *this)
+		&& PayloadDetail::TryCompressImpl(InPayloadType, InPayloadData, *this, CompressionType);
+}
+
+bool FLiveLinkSerializedFrameData::GetPayload(FStructOnScope& OutPayload) const
+{
+	UStruct* PayloadType = nullptr;
+	{
+		FGCScopeGuard Guard;
+		PayloadType = FindObject<UStruct>(nullptr, *PayloadTypeName.ToString());
+	}
+
+	if (PayloadType)
+	{
+		OutPayload.Initialize(PayloadType);
+		const UStruct* PayloadStruct = OutPayload.GetStruct();
+		check(PayloadStruct->IsA<UScriptStruct>());
+		return PayloadDetail::DeserializePreChecks((UScriptStruct*)PayloadStruct, OutPayload.GetStructMemory(), *this)
+			&& PayloadDetail::DeserializeAndDecompress((UScriptStruct*)PayloadStruct, OutPayload.GetStructMemory(), *this);
+	}
+	return false;
+}
+
+bool FLiveLinkSerializedFrameData::GetPayload(const UScriptStruct* InPayloadType, void* InOutPayloadData) const
+{
+	check(InPayloadType && InOutPayloadData);
+	return IsTypeChildOf(InPayloadType)
+		&& PayloadDetail::DeserializePreChecks((UScriptStruct*)InPayloadType, InOutPayloadData, *this)
+		&& PayloadDetail::DeserializeAndDecompress((UScriptStruct*)InPayloadType, InOutPayloadData, *this);
+}
+
+bool FLiveLinkSerializedFrameData::IsTypeChildOf(const UScriptStruct* InPayloadType) const
+{
+	const UStruct* PayloadType = FindObject<UStruct>(nullptr, *PayloadTypeName.ToString());
+	return PayloadType && InPayloadType->IsChildOf(PayloadType);
+}
+
diff --git a/Engine/Source/Runtime/LiveLinkMessageBusFramework/Public/LiveLinkCompression.h b/Engine/Source/Runtime/LiveLinkMessageBusFramework/Public/LiveLinkCompression.h
new file mode 100644
--- /dev/null
+++ b/Engine/Source/Runtime/LiveLinkMessageBusFramework/Public/LiveLinkCompression.h
@@ -0,0 +1,254 @@
+// Copyright Epic Games, Inc. All Rights Reserved.
+
+#pragma once
+
+#include "LiveLinkTypes.h"
+
+#include "LiveLinkCompression.generated.h"
+
+/** Utility struct to serialize a byte array. */
+USTRUCT()
+struct FLiveLinkByteArray
+{
+	GENERATED_BODY()
+
+	bool Serialize(FArchive& Ar)
+	{
+		int32 Num = Bytes.Num();
+		Ar << Num;
+		if (Ar.IsLoading())
+		{
+			Bytes.AddUninitialized(Num);
+		}
+		Ar.Serialize(Bytes.GetData(), Num);
+		return true;
+	}
+
+	UPROPERTY()
+	TArray<uint8> Bytes;
+};
+
+template<>
+struct TStructOpsTypeTraits<FLiveLinkByteArray> : public TStructOpsTypeTraitsBase2<FLiveLinkByteArray>
+{
+	enum
+	{
+		WithSerializer = true,
+	};
+};
+
+/** What compression method should be used on the LiveLink data. */
+UENUM()
+enum class ELiveLinkCompressionMethod : uint8
+{
+	/** Data is uncompressed. */
+	Uncompressed,
+
+	/** Oodle compression is used for the serialized data. */
+	Oodle,
+
+	/**  ZLib compression is used for the serialized data. */
+	Zlib
+};
+
+/** The bias (if any) to use when compressing the data. */
+UENUM()
+enum class ELiveLinkCompressionBias : uint8
+{
+	/** Compress without a bias. */
+	None,
+
+	/** Compress with a bias for size */
+	Size,
+
+	/** Compress with a bias for speed. */
+	Speed
+};
+
+namespace UE::LiveLink::Compression
+{
+
+	/** Get the compression type from the console variable setting. */
+	int32 GetConsoleVariableCompressionType();
+
+	/** Function to check the console variable for the compression flags. Speed vs size. */
+	int32 GetConsoleVariableCompressionFlags();
+
+	/** Based on the size of the data to compress indicate if we should invoke the compressor. */
+	template <typename SizeType>
+	inline bool ShouldCompress(SizeType DataSize)
+	{
+		static_assert(std::is_integral<SizeType>::value);
+		const int32 CompressionLimit = 32;
+		const int32 MinLimit = 512;
+		if (CompressionLimit <= 0)
+		{
+			return DataSize > MinLimit;
+		}
+		const SizeType MaxPackageDataSizeForCompression = static_cast<SizeType>(CompressionLimit) * 1024 * 1024;
+		return DataSize > MinLimit && DataSize <= MaxPackageDataSizeForCompression;
+	}
+
+	/** Get the named compression algorithm to invoke with serializer and memory compressors. */
+	inline FName GetCompressionAlgorithm()
+	{
+		const int32 CompressionAlgo = GetConsoleVariableCompressionType();
+		if (CompressionAlgo == 1)
+		{
+			return NAME_Oodle;
+		}
+		return NAME_Zlib;
+	}
+
+	/** Get the default flags to use when invoking the compressor. */
+	inline ECompressionFlags GetCompressionFlags()
+	{
+		const int32 CompressionFlags = GetConsoleVariableCompressionFlags();
+		if (CompressionFlags == 0)
+		{
+			return COMPRESS_NoFlags;
+		}
+		if (CompressionFlags == 1)
+		{
+			return COMPRESS_BiasSize;
+		}
+		return COMPRESS_BiasSpeed;
+	}
+
+	/** Get the named compression algorithm based on the provided details.  Current this is only Zlib or Oodle. */
+	inline FName GetCompressionAlgorithm(ELiveLinkCompressionMethod Method)
+	{
+		return Method == ELiveLinkCompressionMethod::Oodle ? NAME_Oodle : NAME_Zlib;
+	}
+
+	/** Get the compression enum value given the LiveLink compression settings.  This the value passed into the serializer and memory compressor. */
+	inline ECompressionFlags GetCoreCompressionFlags(ELiveLinkCompressionBias InBias)
+	{
+		if (InBias == ELiveLinkCompressionBias::Size)
+		{
+			return COMPRESS_BiasSize;
+		}
+		else if (InBias == ELiveLinkCompressionBias::Speed)
+		{
+			return COMPRESS_BiasSpeed;
+		}
+		return COMPRESS_NoFlags;
+	}
+
+	/** Get LiveLinkCompressionBias from a ECompressionFlags. */
+	inline ELiveLinkCompressionBias GetCompressionBias(ECompressionFlags Flags)
+	{
+		if (Flags == COMPRESS_BiasMemory)
+		{
+			return ELiveLinkCompressionBias::Size;
+		}
+		else if (COMPRESS_BiasSpeed)
+		{
+			return ELiveLinkCompressionBias::Speed;
+		}
+		else
+		{
+			return ELiveLinkCompressionBias::None;
+		}
+	}
+
+	/** Convert a compression method name to a ELiveLinkCompressionMethod. */
+	inline ELiveLinkCompressionMethod GetCompressionMethod(FName NamedMethod)
+	{
+		check(NamedMethod == NAME_Oodle || NamedMethod == NAME_Zlib);
+		if (NamedMethod == NAME_Oodle)
+		{
+			return ELiveLinkCompressionMethod::Oodle;
+		}
+		else
+		{
+			return ELiveLinkCompressionMethod::Zlib;
+		}
+	}
+}
+
+/** What method should be used for serialziing. */
+UENUM()
+enum class ELiveLinkPayloadSerializationMethod : uint8
+{
+	// The data will be serialized using standard platform method.
+	Standard = 0,
+	// The data will be serialized using Cbor method.
+	Cbor,
+};
+
+/** Dictates when compression should be used. */
+UENUM()
+enum class ELiveLinkPayloadCompressionType : uint8
+{
+	// The serialized data will not be compressed.
+	None = 0,
+	// The serialized data will be compressed based on struct size.
+	Heuristic,
+	// The serialized data will always be compressed.
+	Always
+};
+
+/** Utility struct used to wrap serialized data and handle compression / decompression. */
+USTRUCT()
+struct LIVELINKMESSAGEBUSFRAMEWORK_API FLiveLinkSerializedFrameData
+{
+	GENERATED_BODY()
+
+	FLiveLinkSerializedFrameData() = default;
+
+	FLiveLinkSerializedFrameData(ELiveLinkPayloadSerializationMethod SerializeMethod)
+		: SerializationMethod(SerializeMethod)
+	{
+	}
+
+	/** Initialize this payload from the given data */
+	bool SetPayload(const FStructOnScope& InPayload, ELiveLinkPayloadCompressionType CompressionType = ELiveLinkPayloadCompressionType::Heuristic);
+	bool SetPayload(const UScriptStruct* InPayloadType, const void* InPayloadData, ELiveLinkPayloadCompressionType CompressionType = ELiveLinkPayloadCompressionType::Heuristic);
+
+	template <typename T>
+	bool SetTypedPayload(const T& InPayloadData, ELiveLinkPayloadCompressionType CompressType = ELiveLinkPayloadCompressionType::Heuristic)
+	{
+		return SetPayload(TBaseStructure<T>::Get(), &InPayloadData, CompressType);
+	}
+
+	/** Extract the payload into an in-memory instance */
+	bool GetPayload(FStructOnScope& OutPayload) const;
+	bool GetPayload(const UScriptStruct* InPayloadType, void* InOutPayloadData) const;
+
+	bool IsTypeChildOf(const UScriptStruct* InPayloadType) const;
+
+	template<typename T>
+	bool IsTypeChildOf() const
+	{
+		return IsTypeChildOf(TBaseStructure<T>::Get());
+	}
+
+	template <typename T>
+	bool GetTypedPayload(T& OutPayloadData) const
+	{
+		return GetPayload(TBaseStructure<T>::Get(), &OutPayloadData);
+	}
+
+	/** The typename of the user-defined payload. */
+	UPROPERTY(VisibleAnywhere, Category = "Payload")
+	FName PayloadTypeName;
+
+	/** Specifies the serialization method used to pack the data */
+	UPROPERTY(VisibleAnywhere, Category = "Payload")
+	ELiveLinkPayloadSerializationMethod SerializationMethod = ELiveLinkPayloadSerializationMethod::Standard;
+
+	UPROPERTY(VisibleAnywhere, Category = "Payload")
+	ELiveLinkCompressionMethod CompressionMethod = ELiveLinkCompressionMethod::Uncompressed;
+
+	UPROPERTY(VisibleAnywhere, Category = "Payload")
+	ELiveLinkCompressionBias CompressionBias = ELiveLinkCompressionBias::None;
+
+	/** The uncompressed size of the user-defined payload data. */
+	UPROPERTY(VisibleAnywhere, Category = "Payload")
+	int32 PayloadSize = 0;
+
+	/** The data of the user-defined payload (potentially stored as compressed binary for compact transfer). */
+	UPROPERTY()
+	FLiveLinkByteArray PayloadBytes;
+};
