Fix static analysis warning in LiveLinkCompression.h

#rnx

diff --git a/Engine/Source/Runtime/LiveLinkMessageBusFramework/Public/LiveLinkCompression.h b/Engine/Source/Runtime/LiveLinkMessageBusFramework/Public/LiveLinkCompression.h
--- a/Engine/Source/Runtime/LiveLinkMessageBusFramework/Public/LiveLinkCompression.h
+++ b/Engine/Source/Runtime/LiveLinkMessageBusFramework/Public/LiveLinkCompression.h
@@ -142,7 +142,7 @@ namespace UE::LiveLink::Compression
 		{
 			return ELiveLinkCompressionBias::Size;
 		}
-		else if (COMPRESS_BiasSpeed)
+		else if (Flags == COMPRESS_BiasSpeed)
 		{
 			return ELiveLinkCompressionBias::Speed;
 		}
