LiveLinkHub - Fix LiveLink provider compressing frame data that inherited from LiveLinkAnimationFrameData

#Synced-CL 51376035

diff --git a/Engine/Source/Runtime/LiveLinkMessageBusFramework/Private/LiveLinkProvider.cpp b/Engine/Source/Runtime/LiveLinkMessageBusFramework/Private/LiveLinkProvider.cpp
--- a/Engine/Source/Runtime/LiveLinkMessageBusFramework/Private/LiveLinkProvider.cpp
+++ b/Engine/Source/Runtime/LiveLinkMessageBusFramework/Private/LiveLinkProvider.cpp
@@ -469,7 +469,7 @@ bool FLiveLinkProvider::UpdateSubjectFrameData(const FName SubjectName, FLiveLin
 		TArray<FMessageAddress> Addresses;
 		GetFilteredAddresses(SubjectName, Addresses);
 
-		bool bCanCompressData = CVarLiveLinkCompressAnimData.GetValueOnAnyThread() && FrameData.GetStruct()->IsChildOf(FLiveLinkAnimationFrameData::StaticStruct());
+		bool bCanCompressData = CVarLiveLinkCompressAnimData.GetValueOnAnyThread() && FrameData.GetStruct() == FLiveLinkAnimationFrameData::StaticStruct();
 		TMap<FMessageAddress, bool> SupportsCompression;
 
 		bool bAnyClientSupportsCompression = false;
@@ -486,7 +486,7 @@ bool FLiveLinkProvider::UpdateSubjectFrameData(const FName SubjectName, FLiveLin
 		bCanCompressData = bCanCompressData && bAnyClientSupportsCompression;
 
 		FLiveLinkSerializedFrameData SerializedData;
-		if (bCanCompressData && FrameData.GetStruct()->IsChildOf(FLiveLinkAnimationFrameData::StaticStruct()))
+		if (bCanCompressData && FrameData.GetStruct() == FLiveLinkAnimationFrameData::StaticStruct())
 		{
 			// If this is animation data, serialize and compress it to reduce throughput.
 			SerializedData = SerializeAnimationData(FrameData.Cast<FLiveLinkAnimationFrameData>());
